Copying //UE4/Dev-Anim to Dev-Main (//UE4/Dev-Main) @ 8782600

diff --git a/Engine/Source/Editor/RigVMEditor/Private/RigVMCompiler/RigVMCompiler.cpp b/Engine/Source/Editor/RigVMEditor/Private/RigVMCompiler/RigVMCompiler.cpp
new file mode 100644
--- /dev/null
+++ b/Engine/Source/Editor/RigVMEditor/Private/RigVMCompiler/RigVMCompiler.cpp
@@ -0,0 +1,529 @@
+// Copyright 1998-2019 Epic Games, Inc. All Rights Reserved.
+
+#include "RigVMCompiler/RigVMCompiler.h"
+#include "RigVMModel/RigVMGraphUtils.h"
+#include "RigVMModel/RigVMController.h"
+#include "RigVMEditorModule.h"
+#include "UObject/PropertyPortFlags.h"
+
+FRigVMCompileSettings::FRigVMCompileSettings()
+	: SurpressInfoMessage(true)
+	, SurpressWarnings(false)
+	, SurpressErrors(false)
+{
+}
+
+URigVMCompiler::URigVMCompiler()
+{
+}
+
+bool URigVMCompiler::Compile(URigVMGraph* InGraph, URigVM* OutVM)
+{
+	if (InGraph == nullptr)
+	{
+		ReportError(TEXT("Provided graph is nullptr."));
+		return false;
+	}
+	if (OutVM == nullptr)
+	{
+		ReportError(TEXT("Provided vm is nullptr."));
+		return false;
+	}
+
+	OutVM->Reset();
+
+	FRigVMGraphUtils Utils(InGraph);
+
+	TArray<URigVMNode*> Nodes, Cycle;
+	if (!Utils.TopologicalSort(Nodes, Cycle))
+	{
+		if (Cycle.Num() > 0)
+		{
+			ReportErrorf(TEXT("Cycle detected!"));
+		}
+		return false;
+	}
+
+	FRigVMByteCode& ByteCode = OutVM->ByteCode;
+
+	TMap<FString, FRigVMOperand> PinHashToOperand;
+	for (URigVMNode* Node : Nodes)
+	{
+		if (URigVMStructNode* StructNode = Cast<URigVMStructNode>(Node))
+		{
+			for (URigVMPin* Pin : Node->GetPins())
+			{
+				FRigVMOperand Operand = AddRegisterForPin(Pin, OutVM);
+				PinHashToOperand.Add(GetPinHash(Pin), Operand);
+			}
+		}
+		else if (URigVMParameterNode* ParameterNode = Cast<URigVMParameterNode>(Node))
+		{
+			URigVMPin* ValuePin = Node->FindPin(TEXT("Value"));
+			FString PinHash = GetPinHash(ValuePin);
+			if (!PinHashToOperand.Contains(PinHash))
+			{
+				ERigVMParameterType ParameterType = ParameterNode->IsInput() ? ERigVMParameterType::Input : ERigVMParameterType::Output;
+				FRigVMParameter Parameter;
+
+				FString BaseCPPType = ValuePin->IsArray() ? ValuePin->GetArrayElementCppType() : ValuePin->GetCPPType();
+				UScriptStruct* ScriptStruct = ParameterNode->GetScriptStruct();
+				if (ScriptStruct == nullptr)
+				{
+					ScriptStruct = GetScriptStructForCPPType(BaseCPPType);
+				}
+
+
+				TArray<FString> DefaultValues;
+				if (ValuePin->IsArray())
+				{
+					DefaultValues = URigVMController::SplitDefaultValue(ValuePin->GetDefaultValue());
+				}
+				else
+				{
+					DefaultValues.Add(ValuePin->GetDefaultValue());
+				}
+
+				if (ScriptStruct != nullptr)
+				{
+					TArray<uint8> Data;
+					Data.AddUninitialized(ScriptStruct->GetStructureSize() * DefaultValues.Num());
+					uint8* Ptr = Data.GetData();
+
+					for (FString DefaultValue : DefaultValues)
+					{
+						ScriptStruct->InitializeStruct(Ptr, 1);
+						if (!DefaultValue.IsEmpty())
+						{
+							ScriptStruct->ImportText(*DefaultValue, Ptr, nullptr, PPF_None, nullptr, ScriptStruct->GetName());
+						}
+						Ptr = Ptr + ScriptStruct->GetStructureSize();
+					}
+
+					Parameter = OutVM->AddStructParameter(ParameterType, ParameterNode->GetParameterName(), ScriptStruct, Data.GetData(), DefaultValues.Num());
+					ScriptStruct->DestroyStruct(Data.GetData(), DefaultValues.Num());
+				}
+				else
+				{
+					//Parameter = OutVM->AddPlainParameter(ParameterType, ParameterNode->GetParameterName(), BaseCPPType, nullptr, GetElementSizeFromCPPType(BaseCPPType, ScriptStruct));
+					if (BaseCPPType == TEXT("bool"))
+					{
+						TArray<bool> Values;
+						for (FString DefaultValue : DefaultValues)
+						{
+							Values.Add((DefaultValue == TEXT("True")) || (DefaultValue == TEXT("true")) || (DefaultValue == TEXT("1")));
+						}
+						Parameter = OutVM->AddPlainParameter<bool>(ParameterType, ParameterNode->GetParameterName(), BaseCPPType, Values);
+					}
+					else if (BaseCPPType == TEXT("int32"))
+					{
+						TArray<int32> Values;
+						for (FString DefaultValue : DefaultValues)
+						{
+							if (DefaultValue.IsEmpty())
+							{
+								Values.Add(0);
+							}
+							else
+							{
+								Values.Add(FCString::Atoi(*DefaultValue));
+							}
+						}
+						Parameter = OutVM->AddPlainParameter<int32>(ParameterType, ParameterNode->GetParameterName(), BaseCPPType, Values);
+					}
+					else if (BaseCPPType == TEXT("float"))
+					{
+						TArray<float> Values;
+						for (FString DefaultValue : DefaultValues)
+						{
+							if (DefaultValue.IsEmpty())
+							{
+								Values.Add(0.f);
+							}
+							else
+							{
+								Values.Add(FCString::Atof(*DefaultValue));
+							}
+						}
+						Parameter = OutVM->AddPlainParameter<float>(ParameterType, ParameterNode->GetParameterName(), BaseCPPType, Values);
+					}
+					else if (BaseCPPType == TEXT("FName"))
+					{
+						TArray<FName> Values;
+						for (FString DefaultValue : DefaultValues)
+						{
+							Values.Add(*DefaultValue);
+						}
+						Parameter = OutVM->AddPlainParameter<FName>(ParameterType, ParameterNode->GetParameterName(), BaseCPPType, Values);
+					}
+					else if (BaseCPPType == TEXT("FString"))
+					{
+						Parameter = OutVM->AddPlainParameter<FString>(ParameterType, ParameterNode->GetParameterName(), BaseCPPType, DefaultValues);
+					}
+					else
+					{
+						ensure(false);
+					}
+				}
+
+				FRigVMOperand Operand = OutVM->WorkMemory.GetOperand(Parameter.GetRegisterIndex());
+				PinHashToOperand.Add(PinHash, Operand);
+			}
+		}
+		else if (URigVMVariableNode* VariableNode = Cast<URigVMVariableNode>(Node))
+		{
+			URigVMPin* ValuePin = Node->FindPin(TEXT("Value"));
+			if (ValuePin)
+			{
+				FString PinHash = GetPinHash(ValuePin);
+				if (!PinHashToOperand.Contains(PinHash))
+				{
+					FRigVMOperand Operand = AddRegisterForPin(ValuePin, OutVM);
+					PinHashToOperand.Add(PinHash, Operand);
+				}
+			}
+
+			URigVMPin* ExecutePin = Node->FindPin(TEXT("Execute"));
+			if (ExecutePin)
+			{
+				FString PinHash = GetPinHash(ExecutePin);
+				if (!PinHashToOperand.Contains(PinHash))
+				{
+					FRigVMOperand Operand = AddRegisterForPin(ExecutePin, OutVM);
+					PinHashToOperand.Add(PinHash, Operand);
+				}
+			}
+		}
+		else
+		{
+			ReportErrorf(TEXT("Unsupported node found: '%s'"), *Node->GetClass()->GetName());
+		}
+	}
+
+	for (URigVMNode* Node : Nodes)
+	{
+		if (URigVMStructNode* StructNode = Cast<URigVMStructNode>(Node))
+		{
+			int32 FunctionIndex = OutVM->AddRigVMFunction(StructNode->GetScriptStruct(), StructNode->GetMethodName());
+			if (FunctionIndex == INDEX_NONE)
+			{
+				ReportErrorf(TEXT("Function cannot be found for '%s', method '%s'."), *StructNode->GetScriptStruct()->GetName(), *StructNode->GetMethodName().ToString());
+			}
+
+			TArray<FRigVMOperand> Operands;
+			for (URigVMPin* Pin : Node->GetPins())
+			{
+				FRigVMOperand Operand = PinHashToOperand.FindChecked(GetPinHash(Pin));
+				Operands.Add(Operand);
+			}
+
+			ByteCode.AddExecuteOp(FunctionIndex, Operands);
+		}
+
+		TArray<URigVMPin*> AllPins = Node->GetAllPinsRecursively();
+		for (URigVMPin* Pin : AllPins)
+		{
+			for (URigVMLink* Link : Pin->GetLinks())
+			{
+				if (Link->GetSourcePin() != Pin)
+				{
+					continue;
+				}
+
+				URigVMPin* SourcePin = Link->GetSourcePin();
+				URigVMPin* TargetPin = Link->GetTargetPin();
+				URigVMPin* SourceRootPin = SourcePin->GetRootPin();
+				URigVMPin* TargetRootPin = TargetPin->GetRootPin();
+
+				if (!PinHashToOperand.Contains(GetPinHash(SourceRootPin)))
+				{
+					ensure(false);
+					continue;
+				}
+				if (!PinHashToOperand.Contains(GetPinHash(TargetRootPin)))
+				{
+					ensure(false);
+					continue;
+				}
+
+				FRigVMOperand SourceOperand = PinHashToOperand.FindChecked(GetPinHash(SourceRootPin));
+				FRigVMOperand TargetOperand = PinHashToOperand.FindChecked(GetPinHash(TargetRootPin));
+
+				FString SourceSegmentPath = SourcePin->GetSegmentPath();
+				FString TargetSegmentPath = TargetPin->GetSegmentPath();
+				int32 SourceArrayIndex = 0;
+				int32 TargetArrayIndex = 0;
+
+				if (!SourceSegmentPath.IsEmpty())
+				{
+					if (SourceRootPin->IsArray())
+					{
+						FString Left, Right;
+						if (!SourceSegmentPath.Split(TEXT("."), &Left, &Right))
+						{
+							Left = SourceSegmentPath;
+						}
+						SourceArrayIndex = FCString::Atoi(*Left);
+						SourceSegmentPath = Right;
+					}
+
+					// todo: map to different memories
+					SourceOperand = OutVM->WorkMemory.GetOperand(SourceOperand.GetRegisterIndex(), SourceSegmentPath, SourceArrayIndex);
+				}
+
+				if (!TargetSegmentPath.IsEmpty())
+				{
+					if (TargetRootPin->IsArray())
+					{
+						FString Left, Right;
+						if (!TargetSegmentPath.Split(TEXT("."), &Left, &Right))
+						{
+							Left = TargetSegmentPath;
+						}
+						TargetArrayIndex = FCString::Atoi(*Left);
+						TargetSegmentPath = Right;
+					}
+
+					// todo: map to different memories
+					TargetOperand = OutVM->WorkMemory.GetOperand(TargetOperand.GetRegisterIndex(), TargetSegmentPath, TargetArrayIndex);
+				}
+
+				ByteCode.AddCopyOp(SourceOperand, TargetOperand);
+			}
+		}
+	}
+
+	ByteCode.AddExitOp();
+
+	//URigVMGraph
+	// d) find all literals based on defaults and define them (for now flatten them - no sharing)
+	// g) define all work state (for now no sharing)
+	// h) build out the byte code the the required instructions
+	// i) fix up the jump instructions for all blocks
+
+	return true;
+}
+
+FString URigVMCompiler::GetPinHash(URigVMPin* InPin)
+{
+	URigVMNode* Node = InPin->GetNode();
+	if (URigVMParameterNode* ParameterNode = Cast<URigVMParameterNode>(Node))
+	{
+		if (InPin->GetName() == TEXT("Value"))
+		{
+			return FString::Printf(TEXT("Parameter::%s"), *ParameterNode->GetParameterName().ToString());
+		}
+	}
+	else if (URigVMVariableNode* VariableNode = Cast<URigVMVariableNode>(Node))
+	{
+		if (InPin->GetName() == TEXT("Value"))
+		{
+			return FString::Printf(TEXT("Variable::%s"), *VariableNode->GetVariableName().ToString());
+		}
+	}
+	return InPin->GetPinPath();
+}
+
+UScriptStruct* URigVMCompiler::GetScriptStructForCPPType(const FString& InCPPType)
+{
+	if (InCPPType == TEXT("FRotator"))
+	{
+		return TBaseStructure<FRotator>::Get();
+	}
+	if (InCPPType == TEXT("FQuat"))
+	{
+		return TBaseStructure<FQuat>::Get();
+	}
+	if (InCPPType == TEXT("FTransform"))
+	{
+		return TBaseStructure<FTransform>::Get();
+	}
+	if (InCPPType == TEXT("FLinearColor"))
+	{
+		return TBaseStructure<FLinearColor>::Get();
+	}
+	if (InCPPType == TEXT("FColor"))
+	{
+		return TBaseStructure<FColor>::Get();
+	}
+	if (InCPPType == TEXT("FPlane"))
+	{
+		return TBaseStructure<FPlane>::Get();
+	}
+	if (InCPPType == TEXT("FVector"))
+	{
+		return TBaseStructure<FVector>::Get();
+	}
+	if (InCPPType == TEXT("FVector2D"))
+	{
+		return TBaseStructure<FVector2D>::Get();
+	}
+	if (InCPPType == TEXT("FVector4"))
+	{
+		return TBaseStructure<FVector4>::Get();
+	}
+	return nullptr;
+}
+
+uint16 URigVMCompiler::GetElementSizeFromCPPType(const FString& InCPPType, UScriptStruct* InScriptStruct)
+{
+	if (InScriptStruct == nullptr)
+	{
+		InScriptStruct = GetScriptStructForCPPType(InCPPType);
+	}
+	if (InScriptStruct != nullptr)
+	{
+		return InScriptStruct->GetStructureSize();
+	}
+	if (InCPPType == TEXT("bool"))
+	{
+		return sizeof(bool);
+	}
+	else if (InCPPType == TEXT("int32"))
+	{
+		return sizeof(int32);
+	}
+	if (InCPPType == TEXT("float"))
+	{
+		return sizeof(float);
+	}
+	if (InCPPType == TEXT("FName"))
+	{
+		return sizeof(FName);
+	}
+	if (InCPPType == TEXT("FString"))
+	{
+		return sizeof(FString);
+	}
+
+	ensure(false);
+	return 0;
+}
+
+FRigVMOperand URigVMCompiler::AddRegisterForPin(URigVMPin* InPin, URigVM* OutVM)
+{
+	FString BaseCPPType = InPin->IsArray() ? InPin->GetArrayElementCppType() : InPin->GetCPPType();
+	FString Hash = GetPinHash(InPin);
+	FName HashName = *Hash;
+
+	TArray<FString> DefaultValues;
+	if (InPin->IsArray())
+	{
+		DefaultValues = URigVMController::SplitDefaultValue(InPin->GetDefaultValue());
+	}
+	else
+	{
+		DefaultValues.Add(InPin->GetDefaultValue());
+	}
+
+	UScriptStruct* ScriptStruct = InPin->GetScriptStruct();
+	if (ScriptStruct == nullptr)
+	{
+		ScriptStruct = GetScriptStructForCPPType(BaseCPPType);
+	}
+
+	FRigVMMemoryContainer& WorkMemory = OutVM->WorkMemory;
+	
+	if (ScriptStruct)
+	{
+		TArray<uint8> Data;
+		Data.AddUninitialized(ScriptStruct->GetStructureSize() * DefaultValues.Num());
+		uint8* Ptr = Data.GetData();
+
+		for (FString DefaultValue : DefaultValues)
+		{
+			ScriptStruct->InitializeStruct(Ptr, 1);
+			if (!DefaultValue.IsEmpty())
+			{
+				ScriptStruct->ImportText(*DefaultValue, Ptr, nullptr, PPF_None, nullptr, ScriptStruct->GetName());
+			}
+			Ptr = Ptr + ScriptStruct->GetStructureSize();
+		}
+
+		int32 Register = WorkMemory.AddStructArray(HashName, ScriptStruct, DefaultValues.Num(), Data.GetData(), 1);
+		ScriptStruct->DestroyStruct(Data.GetData(), DefaultValues.Num());
+
+		return WorkMemory.GetOperand(Register);
+	}
+
+	if (BaseCPPType == TEXT("bool"))
+	{
+		TArray<bool> Values;
+		for (FString DefaultValue : DefaultValues)
+		{
+			Values.Add((DefaultValue == TEXT("True")) || (DefaultValue == TEXT("true")) || (DefaultValue == TEXT("1")));
+		}
+		int32 Register = WorkMemory.AddPlainArray<bool>(HashName, Values, 1);
+		return WorkMemory.GetOperand(Register);
+	}
+	else if (BaseCPPType == TEXT("int32"))
+	{
+		TArray<int32> Values;
+		for (FString DefaultValue : DefaultValues)
+		{
+			if (DefaultValue.IsEmpty())
+			{
+				Values.Add(0);
+			}
+			else
+			{
+				Values.Add(FCString::Atoi(*DefaultValue));
+			}
+		}
+		int32 Register = WorkMemory.AddPlainArray<int32>(HashName, Values, 1);
+		return WorkMemory.GetOperand(Register);
+	}
+	if (BaseCPPType == TEXT("float"))
+	{
+		TArray<float> Values;
+		for (FString DefaultValue : DefaultValues)
+		{
+			if (DefaultValue.IsEmpty())
+			{
+				Values.Add(0.f);
+			}
+			else
+			{
+				Values.Add(FCString::Atof(*DefaultValue));
+			}
+		}
+		int32 Register = WorkMemory.AddPlainArray<float>(HashName, Values, 1);
+		return WorkMemory.GetOperand(Register);
+	}
+	if (BaseCPPType == TEXT("FName"))
+	{
+		TArray<FName> Values;
+		for (FString DefaultValue : DefaultValues)
+		{
+			Values.Add(*DefaultValue);
+		}
+		int32 Register = WorkMemory.AddPlainArray<FName>(Values, 1);
+		return WorkMemory.GetOperand(Register);
+	}
+	if (BaseCPPType == TEXT("FString"))
+	{
+		int32 Register = WorkMemory.AddPlainArray<FString>(HashName, DefaultValues, 1);
+		return WorkMemory.GetOperand(Register);
+	}
+
+	ensure(false);
+	return FRigVMOperand();
+}
+
+void URigVMCompiler::ReportInfo(const FString& InMessage)
+{
+	if (!Settings.SurpressInfoMessage)
+	{
+		UE_LOG(LogRigVMEditor, Display, TEXT("%s"), *InMessage);
+	}
+}
+
+void URigVMCompiler::ReportWarning(const FString& InMessage)
+{
+	FScriptExceptionHandler::Get().HandleException(ELogVerbosity::Warning, *InMessage, *FString());
+}
+
+void URigVMCompiler::ReportError(const FString& InMessage)
+{
+	FScriptExceptionHandler::Get().HandleException(ELogVerbosity::Error, *InMessage, *FString());
+}
diff --git a/Engine/Source/Editor/RigVMEditor/Private/RigVMEditorModule.cpp b/Engine/Source/Editor/RigVMEditor/Private/RigVMEditorModule.cpp
new file mode 100644
--- /dev/null
+++ b/Engine/Source/Editor/RigVMEditor/Private/RigVMEditorModule.cpp
@@ -0,0 +1,12 @@
+// Copyright 1998-2019 Epic Games, Inc. All Rights Reserved.
+
+/*=============================================================================
+	RigVMEditor.h: Module implementation.
+=============================================================================*/
+
+#include "RigVMEditorModule.h"
+#include "Modules/ModuleManager.h"
+
+IMPLEMENT_MODULE(FDefaultModuleImpl, RigVMEditor);
+
+DEFINE_LOG_CATEGORY(LogRigVMEditor);
diff --git a/Engine/Source/Editor/RigVMEditor/Private/RigVMModel/Nodes/RigVMRerouteNode.cpp b/Engine/Source/Editor/RigVMEditor/Private/RigVMModel/Nodes/RigVMRerouteNode.cpp
new file mode 100644
--- /dev/null
+++ b/Engine/Source/Editor/RigVMEditor/Private/RigVMModel/Nodes/RigVMRerouteNode.cpp
@@ -0,0 +1,16 @@
+// Copyright 1998-2019 Epic Games, Inc. All Rights Reserved.
+
+#include "RigVMModel/Nodes/RigVMRerouteNode.h"
+
+const FString URigVMRerouteNode::RerouteName = TEXT("Reroute");
+const FString URigVMRerouteNode::ValueName = TEXT("Value");
+
+URigVMRerouteNode::URigVMRerouteNode()
+: bShowAsFullNode(true)
+{
+}
+
+bool URigVMRerouteNode::GetShowsAsFullNode() const
+{
+	return bShowAsFullNode;
+}
\ No newline at end of file
diff --git a/Engine/Source/Editor/RigVMEditor/Private/RigVMModel/RigVMController.cpp b/Engine/Source/Editor/RigVMEditor/Private/RigVMModel/RigVMController.cpp
new file mode 100644
--- /dev/null
+++ b/Engine/Source/Editor/RigVMEditor/Private/RigVMModel/RigVMController.cpp
@@ -0,0 +1,2257 @@
+// Copyright 1998-2019 Epic Games, Inc. All Rights Reserved.
+
+#include "RigVMModel/RigVMController.h"
+#include "RigVMModel/RigVMControllerActions.h"
+#include "RigVMCore/RigVMRegistry.h"
+#include "RigVMCore/RigVMExecuteContext.h"
+#include "RigVMEditorModule.h"
+#include "UObject/PropertyPortFlags.h"
+#include "UObject/Package.h"
+#include "Misc/CoreMisc.h"
+
+int32 URigVMController::sRemovedPinIndex = 0;
+
+URigVMController::URigVMController()
+	: bReportWarningsAndErrors(true)
+{
+}
+
+URigVMController::URigVMController(const FObjectInitializer& ObjectInitializer)
+	: Super(ObjectInitializer)
+	, bReportWarningsAndErrors(true)
+{
+	ActionStack = CreateDefaultSubobject<URigVMActionStack>(TEXT("ActionStack"));
+}
+
+URigVMController::~URigVMController()
+{
+	SetGraph(nullptr);
+}
+
+URigVMGraph* URigVMController::GetGraph() const
+{
+	return Graph;
+}
+
+void URigVMController::SetGraph(URigVMGraph* InGraph)
+{
+	if (Graph)
+	{
+		Graph->OnModified().RemoveAll(this);
+	}
+
+	Graph = InGraph;
+
+	if (Graph)
+	{
+		Graph->OnModified().AddUObject(this, &URigVMController::HandleModifiedEvent);
+	}
+
+	HandleModifiedEvent(ERigVMGraphNotifType::GraphChanged, Graph, nullptr);
+}
+
+FRigVMGraphModifiedEvent& URigVMController::OnModified()
+{
+	return ModifiedEventStatic;
+}
+
+void URigVMController::Notify(ERigVMGraphNotifType InNotifType, URigVMGraph* InGraph, UObject* InSubject)
+{
+	if (Graph)
+	{
+		Graph->Notify(InNotifType, InGraph, InSubject);
+	}
+}
+
+void URigVMController::HandleModifiedEvent(ERigVMGraphNotifType InNotifType, URigVMGraph* InGraph, UObject* InSubject)
+{
+	ModifiedEventStatic.Broadcast(InNotifType, InGraph, InSubject);
+	if (ModifiedEventDynamic.IsBound())
+	{
+		ModifiedEventDynamic.Broadcast(InNotifType, InGraph, InSubject);
+	}
+}
+
+#if WITH_EDITOR
+
+URigVMStructNode* URigVMController::AddStructNode(UScriptStruct* InScriptStruct, const FName& InMethodName, const FVector2D& InPosition, const FString& InNodeName, bool bUndo)
+{
+	if(!IsValidGraph())
+	{
+		return nullptr;
+	}
+	if (InScriptStruct == nullptr)
+	{
+		ReportError(TEXT("InScriptStruct is null."));
+		return nullptr;
+	}
+	if (InMethodName == NAME_None)
+	{
+		ReportError(TEXT("InMethodName is None."));
+		return nullptr;
+	}
+
+	FString FunctionName = FString::Printf(TEXT("F%s::%s"), *InScriptStruct->GetName(), *InMethodName.ToString());
+	FRigVMFunctionPtr Function = FRigVMRegistry::Get().Find(*FunctionName);
+	if (Function == nullptr)
+	{
+		ReportErrorf(TEXT("RIGVM_METHOD '%s' cannot be found. is None."), *FunctionName);
+		return nullptr;
+	}
+
+	FString Name = GetValidName(InNodeName.IsEmpty() ? InScriptStruct->GetName() : InNodeName);
+	URigVMStructNode* Node = NewObject<URigVMStructNode>(Graph, *Name);
+	Node->ScriptStruct = InScriptStruct;
+	Node->MethodName = InMethodName;
+	Node->Position = InPosition;
+	Node->NodeTitle = InScriptStruct->GetMetaData(TEXT("DisplayName"));
+	
+	FString NodeColorMetadata;
+	InScriptStruct->GetStringMetaDataHierarchical(*URigVMNode::NodeColorName, &NodeColorMetadata);
+	if (!NodeColorMetadata.IsEmpty())
+	{
+		Node->NodeColor = GetColorFromMetadata(NodeColorMetadata);
+	}
+
+	TArray<uint8> TempBuffer;
+	TempBuffer.AddUninitialized(InScriptStruct->GetStructureSize());
+	InScriptStruct->InitializeDefaultValue(TempBuffer.GetData());
+
+	FString ExportedDefaultValue;
+	InScriptStruct->ExportText(ExportedDefaultValue, TempBuffer.GetData(), nullptr, nullptr, PPF_None, nullptr);
+	InScriptStruct->DestroyStruct(TempBuffer.GetData());
+	AddPinsForStruct(InScriptStruct, Node, nullptr, ERigVMPinDirection::Invalid, ExportedDefaultValue);
+
+	Graph->Nodes.Add(Node);
+	Graph->MarkPackageDirty();
+
+	FRigVMAddStructNodeAction Action;
+	if (bUndo)
+	{
+		Action = FRigVMAddStructNodeAction(Node);
+		Action.Title = FString::Printf(TEXT("Add %s Node"), *Node->GetNodeTitle());
+		ActionStack->BeginAction(Action);
+	}
+
+	Notify(ERigVMGraphNotifType::NodeAdded, Graph, Node);
+
+	if (bUndo)
+	{
+		ActionStack->EndAction(Action);
+	}
+
+	return Node;
+}
+
+URigVMStructNode* URigVMController::AddStructNodeFromStructPath(const FString& InScriptStructPath, const FName& InMethodName, const FVector2D& InPosition, const FString& InNodeName, bool bUndo)
+{
+	if(!IsValidGraph())
+	{
+		return nullptr;
+	}
+
+	UScriptStruct* ScriptStruct = FindObject<UScriptStruct>(ANY_PACKAGE, *InScriptStructPath);
+	if (ScriptStruct == nullptr)
+	{
+		ReportErrorf(TEXT("Cannot find struct for path '%s'."), *InScriptStructPath);
+		return nullptr;
+	}
+
+	return AddStructNode(ScriptStruct, InMethodName, InPosition, InNodeName, bUndo);
+}
+
+URigVMVariableNode* URigVMController::AddVariableNode(const FName& InVariableName, const FString& InCPPType, UScriptStruct* InScriptStruct, bool bIsGetter, const FString& InDefaultValue, const FVector2D& InPosition, const FString& InNodeName, bool bUndo)
+{
+	if (!IsValidGraph())
+	{
+		return nullptr;
+	}
+
+	TArray<FRigVMGraphVariableDescription> ExistingVariables = Graph->GetVariableDescriptions();
+	for (const FRigVMGraphVariableDescription& ExistingVariable : ExistingVariables)
+	{
+		if (ExistingVariable.Name == InVariableName)
+		{
+			if (ExistingVariable.CPPType != InCPPType ||
+				ExistingVariable.ScriptStruct != InScriptStruct)
+			{
+				ReportErrorf(TEXT("Cannot add variable '%s' - variable already exists."), *InVariableName.ToString());
+				return nullptr;
+			}
+		}
+	}
+
+	FString Name = GetValidName(InNodeName.IsEmpty() ? FString(TEXT("VariableNode")) : InNodeName);
+	URigVMVariableNode* Node = NewObject<URigVMVariableNode>(Graph, *Name);
+	Node->Position = InPosition;
+	Node->VariableName = InVariableName;
+
+	if (!bIsGetter)
+	{
+		URigVMPin* ExecutePin = NewObject<URigVMPin>(Node, *URigVMNode::ExecuteName);
+		ExecutePin->CPPType = TEXT("FRigVMExecuteContext");
+		ExecutePin->ScriptStruct = FRigVMExecuteContext::StaticStruct();
+		ExecutePin->ScriptStructPath = *ExecutePin->ScriptStruct->GetPathName();
+		ExecutePin->Direction = ERigVMPinDirection::IO;
+		Node->Pins.Add(ExecutePin);
+	}
+
+	URigVMPin* ValuePin = NewObject<URigVMPin>(Node, *URigVMVariableNode::ValueName);
+	ValuePin->CPPType = InCPPType;
+	if (InScriptStruct)
+	{
+		ValuePin->ScriptStruct = InScriptStruct;
+		ValuePin->ScriptStructPath = *ValuePin->ScriptStruct->GetPathName();
+	}
+	ValuePin->Direction = bIsGetter ? ERigVMPinDirection::Output : ERigVMPinDirection::Input;
+	Node->Pins.Add(ValuePin);
+
+	Graph->Nodes.Add(Node);
+
+	if (ValuePin->IsStruct())
+	{
+		FString DefaultValue = InDefaultValue;
+		if (DefaultValue.IsEmpty() || DefaultValue == TEXT("()"))
+		{
+			TArray<uint8> TempBuffer;
+			TempBuffer.AddUninitialized(ValuePin->ScriptStruct->GetStructureSize());
+			ValuePin->ScriptStruct->InitializeDefaultValue(TempBuffer.GetData());
+			ValuePin->ScriptStruct->ExportText(DefaultValue, TempBuffer.GetData(), nullptr, nullptr, PPF_None, nullptr);
+			ValuePin->ScriptStruct->DestroyStruct(TempBuffer.GetData());
+		}
+		AddPinsForStruct(ValuePin->ScriptStruct, Node, ValuePin, ValuePin->Direction, DefaultValue);
+	}
+	else if (!InDefaultValue.IsEmpty() && InDefaultValue != TEXT("()"))
+	{
+		SetPinDefaultValue(ValuePin, InDefaultValue, true, false, false);
+	}
+
+
+	Graph->MarkPackageDirty();
+
+	FRigVMAddVariableNodeAction Action;
+	if (bUndo)
+	{
+		Action = FRigVMAddVariableNodeAction(Node);
+		Action.Title = FString::Printf(TEXT("Add %s Variable"), *InVariableName.ToString());
+		ActionStack->BeginAction(Action);
+	}
+
+	Notify(ERigVMGraphNotifType::NodeAdded, Graph, Node);
+
+	if (bUndo)
+	{
+		ActionStack->EndAction(Action);
+	}
+
+	return Node;
+}
+
+URigVMVariableNode* URigVMController::AddVariableNodeFromStructPath(const FName& InVariableName, const FString& InCPPType, const FString& InScriptStructPath, bool bIsGetter, const FString& InDefaultValue, const FVector2D& InPosition, const FString& InNodeName, bool bUndo)
+{
+	if (!IsValidGraph())
+	{
+		return nullptr;
+	}
+
+	UScriptStruct* ScriptStruct = nullptr;
+	if (!InScriptStructPath.IsEmpty())
+	{
+		ScriptStruct = FindObject<UScriptStruct>(ANY_PACKAGE, *InScriptStructPath);
+		if (ScriptStruct == nullptr)
+		{
+			ReportErrorf(TEXT("Cannot find struct for path '%s'."), *InScriptStructPath);
+			return nullptr;
+		}
+	}
+
+	return AddVariableNode(InVariableName, InCPPType, ScriptStruct, bIsGetter, InDefaultValue, InPosition, InNodeName, bUndo);
+}
+
+URigVMParameterNode* URigVMController::AddParameterNode(const FName& InParameterName, const FString& InCPPType, UScriptStruct* InScriptStruct, bool bIsInput, const FString& InDefaultValue, const FVector2D& InPosition, const FString& InNodeName, bool bUndo)
+{
+	if (!IsValidGraph())
+	{
+		return nullptr;
+	}
+
+	TArray<FRigVMGraphParameterDescription> ExistingParameters = Graph->GetParameterDescriptions();
+	for (const FRigVMGraphParameterDescription& ExistingParameter : ExistingParameters)
+	{
+		if (ExistingParameter.Name == InParameterName)
+		{
+			if (ExistingParameter.CPPType != InCPPType ||
+				ExistingParameter.ScriptStruct != InScriptStruct ||
+				ExistingParameter.bIsInput != bIsInput)
+			{
+				ReportErrorf(TEXT("Cannot add parameter '%s' - parameter already exists."), *InParameterName.ToString());
+				return nullptr;
+			}
+		}
+	}
+
+	FString Name = GetValidName(InNodeName.IsEmpty() ? FString(TEXT("ParameterNode")) : InNodeName);
+	URigVMParameterNode* Node = NewObject<URigVMParameterNode>(Graph, *Name);
+	Node->Position = InPosition;
+	Node->ParameterName = InParameterName;
+
+	URigVMPin* ValuePin = NewObject<URigVMPin>(Node, *URigVMParameterNode::ValueName);
+	ValuePin->CPPType = InCPPType;
+	if (InScriptStruct)
+	{
+		ValuePin->ScriptStruct = InScriptStruct;
+		ValuePin->ScriptStructPath = *ValuePin->ScriptStruct->GetPathName();
+	}
+	ValuePin->Direction = bIsInput ? ERigVMPinDirection::Output : ERigVMPinDirection::Input;
+	Node->Pins.Add(ValuePin);
+
+	Graph->Nodes.Add(Node);
+
+	if (ValuePin->IsStruct())
+	{
+		FString DefaultValue = InDefaultValue;
+		if(DefaultValue.IsEmpty() || DefaultValue == TEXT("()"))
+		{
+			TArray<uint8> TempBuffer;
+			TempBuffer.AddUninitialized(ValuePin->ScriptStruct->GetStructureSize());
+			ValuePin->ScriptStruct->InitializeDefaultValue(TempBuffer.GetData());
+			ValuePin->ScriptStruct->ExportText(DefaultValue, TempBuffer.GetData(), nullptr, nullptr, PPF_None, nullptr);
+			ValuePin->ScriptStruct->DestroyStruct(TempBuffer.GetData());
+		}
+		AddPinsForStruct(ValuePin->ScriptStruct, Node, ValuePin, ValuePin->Direction, DefaultValue);
+	}
+	else if (!InDefaultValue.IsEmpty() && InDefaultValue != TEXT("()"))
+	{
+		SetPinDefaultValue(ValuePin, InDefaultValue, true, false, false);
+	}
+
+	Graph->MarkPackageDirty();
+
+	FRigVMAddParameterNodeAction Action;
+	if (bUndo)
+	{
+		Action = FRigVMAddParameterNodeAction(Node);
+		Action.Title = FString::Printf(TEXT("Add %s Parameter"), *InParameterName.ToString());
+		ActionStack->BeginAction(Action);
+	}
+
+	Notify(ERigVMGraphNotifType::NodeAdded, Graph, Node);
+
+	if (bUndo)
+	{
+		ActionStack->EndAction(Action);
+	}
+
+	return Node;
+}
+
+URigVMParameterNode* URigVMController::AddParameterNodeFromStructPath(const FName& InParameterName, const FString& InCPPType, const FString& InScriptStructPath, bool bIsInput, const FString& InDefaultValue, const FVector2D& InPosition, const FString& InNodeName, bool bUndo)
+{
+	if (!IsValidGraph())
+	{
+		return nullptr;
+	}
+
+	UScriptStruct* ScriptStruct = nullptr;
+	if (!InScriptStructPath.IsEmpty())
+	{
+		ScriptStruct = FindObject<UScriptStruct>(ANY_PACKAGE, *InScriptStructPath);
+		if (ScriptStruct == nullptr)
+		{
+			ReportErrorf(TEXT("Cannot find struct for path '%s'."), *InScriptStructPath);
+			return nullptr;
+		}
+	}
+
+	return AddParameterNode(InParameterName, InCPPType, ScriptStruct, bIsInput, InDefaultValue, InPosition, InNodeName, bUndo);
+}
+
+URigVMCommentNode* URigVMController::AddCommentNode(const FString& InCommentText, const FVector2D& InPosition, const FVector2D& InSize, const FLinearColor& InColor, const FString& InNodeName, bool bUndo)
+{
+	if (!IsValidGraph())
+	{
+		return nullptr;
+	}
+
+	FString Name = GetValidName(InNodeName.IsEmpty() ? FString(TEXT("CommentNode")) : InNodeName);
+	URigVMCommentNode* Node = NewObject<URigVMCommentNode>(Graph, *Name);
+	Node->Position = InPosition;
+	Node->Size = InSize;
+	Node->NodeColor = InColor;
+	Node->CommentText = InCommentText;
+
+	Graph->Nodes.Add(Node);
+	Graph->MarkPackageDirty();
+
+	FRigVMAddCommentNodeAction Action;
+	if (bUndo)
+	{
+		Action = FRigVMAddCommentNodeAction(Node);
+		Action.Title = FString::Printf(TEXT("Add Comment"));
+		ActionStack->BeginAction(Action);
+	}
+
+	Notify(ERigVMGraphNotifType::NodeAdded, Graph, Node);
+
+	if (bUndo)
+	{
+		ActionStack->EndAction(Action);
+	}
+
+	return Node;
+}
+
+URigVMRerouteNode* URigVMController::AddRerouteNodeOnLink(URigVMLink* InLink, bool bShowAsFullNode, const FVector2D& InPosition, const FString& InNodeName, bool bUndo)
+{
+	if(!IsValidLinkForGraph(InLink))
+	{
+		return nullptr;
+	}
+
+	URigVMPin* SourcePin = InLink->GetSourcePin();
+	URigVMPin* TargetPin = InLink->GetTargetPin();
+
+	FRigVMBaseAction Action;
+	if (bUndo)
+	{
+		Action.Title = FString::Printf(TEXT("Add Reroute"));
+		ActionStack->BeginAction(Action);
+	}
+
+	URigVMRerouteNode* Node = AddRerouteNodeOnPin(TargetPin->GetPinPath(), true, bShowAsFullNode, InPosition, InNodeName, bUndo);
+	if (Node == nullptr)
+	{
+		if (bUndo)
+		{
+			ActionStack->CancelAction(Action);
+		}
+		return nullptr;
+	}
+
+	URigVMPin* ValuePin = Node->Pins[0];
+	AddLink(SourcePin, ValuePin, bUndo);
+
+	if (bUndo)
+	{
+		ActionStack->EndAction(Action);
+	}
+
+	return Node;
+}
+
+URigVMRerouteNode* URigVMController::AddRerouteNodeOnLinkPath(const FString& InLinkPinPathRepresentation, bool bShowAsFullNode, const FVector2D& InPosition, const FString& InNodeName, bool bUndo)
+{
+	if (!IsValidGraph())
+	{
+		return nullptr;
+	}
+
+	URigVMLink* Link = Graph->FindLink(InLinkPinPathRepresentation);
+	return AddRerouteNodeOnLink(Link, bShowAsFullNode, InPosition, InNodeName, bUndo);
+}
+
+URigVMRerouteNode* URigVMController::AddRerouteNodeOnPin(const FString& InPinPath, bool bAsInput, bool bShowAsFullNode, const FVector2D& InPosition, const FString& InNodeName, bool bUndo)
+{
+	if (!IsValidGraph())
+	{
+		return nullptr;
+	}
+
+	URigVMPin* Pin = Graph->FindPin(InPinPath);
+	if(Pin == nullptr)
+	{
+		return nullptr;
+	}
+
+	FRigVMBaseAction Action;
+	if (bUndo)
+	{
+		Action.Title = FString::Printf(TEXT("Add Reroute"));
+		ActionStack->BeginAction(Action);
+	}
+
+	BreakAllLinks(Pin, bAsInput, bUndo);
+
+	FString Name = GetValidName(InNodeName.IsEmpty() ? FString(TEXT("RerouteNode")) : InNodeName);
+	URigVMRerouteNode* Node = NewObject<URigVMRerouteNode>(Graph, *Name);
+	Node->Position = InPosition;
+	Node->bShowAsFullNode = bShowAsFullNode;
+
+	URigVMPin* ValuePin = NewObject<URigVMPin>(Node, *URigVMRerouteNode::ValueName);
+	ConfigurePinFromPin(ValuePin, Pin);
+	ValuePin->Direction = ERigVMPinDirection::IO;
+	Node->Pins.Add(ValuePin);
+
+	if (ValuePin->IsStruct())
+	{
+		AddPinsForStruct(ValuePin->GetScriptStruct(), Node, ValuePin, ValuePin->Direction, FString());
+	}
+
+	FString DefaultValue = Pin->GetDefaultValue();
+	if (!DefaultValue.IsEmpty())
+	{
+		SetPinDefaultValue(ValuePin, Pin->GetDefaultValue(), true, false, false);
+	}
+
+	Graph->Nodes.Add(Node);
+	Graph->MarkPackageDirty();
+
+	if (bUndo)
+	{
+		ActionStack->AddAction(FRigVMAddRerouteNodeAction(Node, InPinPath, bAsInput));
+	}
+
+	Notify(ERigVMGraphNotifType::NodeAdded, Graph, Node);
+
+	if (bAsInput)
+	{
+		AddLink(ValuePin, Pin, bUndo);
+	}
+	else
+	{
+		AddLink(Pin, ValuePin, bUndo);
+	}
+
+	if (bUndo)
+	{
+		ActionStack->EndAction(Action);
+	}
+
+	return Node;
+}
+
+bool URigVMController::Undo()
+{
+	if (!IsValidGraph())
+	{
+		return false;
+	}
+	return ActionStack->Undo(this);
+}
+
+bool URigVMController::Redo()
+{
+	if (!IsValidGraph())
+	{
+		return false;
+	}
+	return ActionStack->Redo(this);
+}
+
+bool URigVMController::OpenUndoBracket(const FString& InTitle)
+{
+	if (!IsValidGraph())
+	{
+		return false;
+	}
+	return ActionStack->OpenUndoBracket(InTitle);
+}
+
+bool URigVMController::CloseUndoBracket()
+{
+	if (!IsValidGraph())
+	{
+		return false;
+	}
+	return ActionStack->CloseUndoBracket();
+}
+
+#endif
+
+bool URigVMController::RemoveNode(URigVMNode* InNode, bool bUndo)
+{
+	if (!IsValidNodeForGraph(InNode))
+	{
+		return false;
+	}
+
+	FRigVMRemoveNodeAction Action;
+	if (bUndo)
+	{
+		Action = FRigVMRemoveNodeAction(InNode);
+		Action.Title = FString::Printf(TEXT("Remove %s Node"), *InNode->GetNodeTitle());
+		ActionStack->BeginAction(Action);
+	}
+
+	SelectNode(InNode, false, bUndo);
+
+	for (URigVMPin* Pin : InNode->Pins)
+	{
+		BreakAllLinks(Pin, true, bUndo);
+		BreakAllLinks(Pin, false, bUndo);
+		BreakAllLinksRecursive(Pin, true, false, bUndo);
+		BreakAllLinksRecursive(Pin, false, false, bUndo);
+	}
+
+	Graph->Nodes.Remove(InNode);
+	Graph->MarkPackageDirty();
+
+	Notify(ERigVMGraphNotifType::NodeRemoved, Graph, InNode);
+
+	InNode->MarkPendingKill();
+
+	if (bUndo)
+	{
+		ActionStack->EndAction(Action);
+	}
+
+	return true;
+}
+
+bool URigVMController::RemoveNodeByName(const FName& InNodeName, bool bUndo)
+{
+	if (!IsValidGraph())
+	{
+		return false;
+	}
+	return RemoveNode(Graph->FindNodeByName(InNodeName), bUndo);
+}
+
+bool URigVMController::SelectNode(URigVMNode* InNode, bool bSelect, bool bUndo)
+{
+	if (!IsValidNodeForGraph(InNode))
+	{
+		return false;
+	}
+
+	if (InNode->IsSelected() == bSelect)
+	{
+		return false;
+	}
+
+	FRigVMSelectNodeAction Action;
+	if (bUndo)
+	{
+		Action = FRigVMSelectNodeAction(InNode);
+		if (bSelect)
+		{
+			Action.Title = FString::Printf(TEXT("Select %s Node"), *InNode->GetNodeTitle());
+		}
+		else
+		{
+			Action.Title = FString::Printf(TEXT("Deselect %s Node"), *InNode->GetNodeTitle());
+		}
+		ActionStack->BeginAction(Action);
+	}
+
+	if (bSelect)
+	{
+		Graph->SelectedNodes.Add(InNode->GetFName());
+		Notify(ERigVMGraphNotifType::NodeSelected, Graph, InNode);
+	}
+	else
+	{
+		Graph->SelectedNodes.Remove(InNode->GetFName());
+		Notify(ERigVMGraphNotifType::NodeDeselected, Graph, InNode);
+	}
+
+	if (bUndo)
+	{
+		ActionStack->EndAction(Action);
+	}
+
+	return true;
+}
+
+bool URigVMController::SelectNodeByName(const FName& InNodeName, bool bSelect, bool bUndo)
+{
+	if (!IsValidGraph())
+	{
+		return false;
+	}
+	return SelectNode(Graph->FindNodeByName(InNodeName), bSelect, bUndo);
+}
+
+bool URigVMController::ClearNodeSelection(bool bUndo)
+{
+	if (!IsValidGraph())
+	{
+		return false;
+	}
+
+	FRigVMBaseAction Action;
+	if (bUndo)
+	{
+		Action.Title = FString::Printf(TEXT("Clear selection"));
+		ActionStack->BeginAction(Action);
+	}
+
+	TArray<FName> Selection = Graph->SelectedNodes;
+	for (const FName& SelectedNode : Selection)
+	{
+		SelectNodeByName(SelectedNode, false, bUndo);
+	}
+
+	if (bUndo)
+	{
+		ActionStack->EndAction(Action);
+	}
+
+	return Selection.Num() > 0;
+}
+
+bool URigVMController::SetNodePosition(URigVMNode* InNode, const FVector2D& InPosition, bool bUndo, bool bMergeUndoAction)
+{
+	if (!IsValidNodeForGraph(InNode))
+	{
+		return false;
+	}
+
+	FRigVMSetNodePositionAction Action;
+	if (bUndo)
+	{
+		Action = FRigVMSetNodePositionAction(InNode, InPosition);
+		Action.Title = FString::Printf(TEXT("Set Node Position"));
+		ActionStack->BeginAction(Action);
+	}
+
+	InNode->Position = InPosition;
+	Notify(ERigVMGraphNotifType::NodePositionChanged, Graph, InNode);
+
+	if (bUndo)
+	{
+		ActionStack->EndAction(Action, bMergeUndoAction);
+	}
+
+	return true;
+}
+
+bool URigVMController::SetNodePositionByName(const FName& InNodeName, const FVector2D& InPosition, bool bUndo, bool bMergeUndoAction)
+{
+	if(!IsValidGraph())
+	{
+		return false;
+	}
+
+	URigVMNode* Node = Graph->FindNodeByName(InNodeName);
+	return SetNodePosition(Node, InPosition, bUndo, bMergeUndoAction);
+}
+
+bool URigVMController::SetNodeSize(URigVMNode* InNode, const FVector2D& InSize, bool bUndo, bool bMergeUndoAction)
+{
+	if (!IsValidNodeForGraph(InNode))
+	{
+		return false;
+	}
+
+	FRigVMSetNodeSizeAction Action;
+	if (bUndo)
+	{
+		Action = FRigVMSetNodeSizeAction(InNode, InSize);
+		Action.Title = FString::Printf(TEXT("Set Node Size"));
+		ActionStack->BeginAction(Action);
+	}
+
+	InNode->Size = InSize;
+	Notify(ERigVMGraphNotifType::NodeSizeChanged, Graph, InNode);
+
+	if (bUndo)
+	{
+		ActionStack->EndAction(Action, bMergeUndoAction);
+	}
+
+	return true;
+}
+
+bool URigVMController::SetNodeSizeByName(const FName& InNodeName, const FVector2D& InSize, bool bUndo, bool bMergeUndoAction)
+{
+	if(!IsValidGraph())
+	{
+		return false;
+	}
+
+	URigVMNode* Node = Graph->FindNodeByName(InNodeName);
+	return SetNodeSize(Node, InSize, bUndo, bMergeUndoAction);
+}
+
+bool URigVMController::SetNodeColor(URigVMNode* InNode, const FLinearColor& InColor, bool bUndo, bool bMergeUndoAction)
+{
+	if (!IsValidNodeForGraph(InNode))
+	{
+		return false;
+	}
+
+	FRigVMSetNodeColorAction Action;
+	if (bUndo)
+	{
+		Action = FRigVMSetNodeColorAction(InNode, InColor);
+		Action.Title = FString::Printf(TEXT("Set Node Color"));
+		ActionStack->BeginAction(Action);
+	}
+
+	InNode->NodeColor = InColor;
+	Notify(ERigVMGraphNotifType::NodeColorChanged, Graph, InNode);
+
+	if (bUndo)
+	{
+		ActionStack->EndAction(Action, bMergeUndoAction);
+	}
+
+	return true;
+}
+
+bool URigVMController::SetNodeColorByName(const FName& InNodeName, const FLinearColor& InColor, bool bUndo, bool bMergeUndoAction)
+{
+	if(!IsValidGraph())
+	{
+		return false;
+	}
+
+	URigVMNode* Node = Graph->FindNodeByName(InNodeName);
+	return SetNodeColor(Node, InColor, bUndo, bMergeUndoAction);
+}
+
+bool URigVMController::SetCommentText(URigVMNode* InNode, const FString& InCommentText, bool bUndo)
+{
+	if (!IsValidNodeForGraph(InNode))
+	{
+		return false;
+	}
+
+	if (URigVMCommentNode* CommentNode = Cast<URigVMCommentNode>(InNode))
+	{
+		FRigVMSetCommentTextAction Action;
+		if (bUndo)
+		{
+			Action = FRigVMSetCommentTextAction(CommentNode, InCommentText);
+			Action.Title = FString::Printf(TEXT("Set Comment Text"));
+			ActionStack->BeginAction(Action);
+		}
+
+		CommentNode->CommentText = InCommentText;
+		Notify(ERigVMGraphNotifType::CommentTextChanged, Graph, InNode);
+
+		if (bUndo)
+		{
+			ActionStack->EndAction(Action);
+		}
+
+		return true;
+	}
+
+	return false;
+}
+
+bool URigVMController::SetCommentTextByName(const FName& InNodeName, const FString& InCommentText, bool bUndo)
+{
+	if(!IsValidGraph())
+	{
+		return false;
+	}
+
+	URigVMNode* Node = Graph->FindNodeByName(InNodeName);
+	return SetCommentText(Node, InCommentText, bUndo);
+}
+
+bool URigVMController::SetRerouteCompactness(URigVMNode* InNode, bool bShowAsFullNode, bool bUndo)
+{
+	if (!IsValidNodeForGraph(InNode))
+	{
+		return false;
+	}
+
+	if (URigVMRerouteNode* RerouteNode = Cast<URigVMRerouteNode>(InNode))
+	{
+		FRigVMSetRerouteCompactnessAction Action;
+		if (bUndo)
+		{
+			Action = FRigVMSetRerouteCompactnessAction(RerouteNode, bShowAsFullNode);
+			Action.Title = FString::Printf(TEXT("Set Reroute Size"));
+			ActionStack->BeginAction(Action);
+		}
+
+		RerouteNode->bShowAsFullNode = bShowAsFullNode;
+		Notify(ERigVMGraphNotifType::RerouteCompactnessChanged, Graph, InNode);
+
+		if (bUndo)
+		{
+			ActionStack->EndAction(Action);
+		}
+
+		return true;
+	}
+
+	return false;
+}
+
+bool URigVMController::SetRerouteCompactnessByName(const FName& InNodeName, bool bShowAsFullNode, bool bUndo)
+{
+	if(!IsValidGraph())
+	{
+		return false;
+	}
+
+	URigVMNode* Node = Graph->FindNodeByName(InNodeName);
+	return SetRerouteCompactness(Node, bShowAsFullNode, bUndo);
+}
+
+bool URigVMController::RenameVariable(const FName& InOldName, const FName& InNewName, bool bUndo)
+{
+	if(!IsValidGraph())
+	{
+		return false;
+	}
+
+	if (InOldName == InNewName)
+	{
+		ReportWarning(TEXT("RenameVariable: InOldName and InNewName are equal."));
+		return false;
+	}
+
+	TArray<FRigVMGraphVariableDescription> ExistingVariables = Graph->GetVariableDescriptions();
+	for (const FRigVMGraphVariableDescription& ExistingVariable : ExistingVariables)
+	{
+		if (ExistingVariable.Name == InNewName)
+		{
+			ReportErrorf(TEXT("Cannot rename variable to '%s' - variable already exists."), *InNewName.ToString());
+			return false;
+		}
+	}
+
+	FRigVMRenameVariableAction Action;
+	if (bUndo)
+	{
+		Action = FRigVMRenameVariableAction(InOldName, InNewName);
+		Action.Title = FString::Printf(TEXT("Rename Variable"));
+		ActionStack->BeginAction(Action);
+	}
+
+	TArray<URigVMNode*> RenamedNodes;
+	for (URigVMNode* Node : Graph->Nodes)
+	{
+		if(URigVMVariableNode* VariableNode = Cast<URigVMVariableNode>(Node))
+		{
+			if (VariableNode->VariableName == InOldName)
+			{
+				VariableNode->VariableName = InNewName;
+				RenamedNodes.Add(Node);
+			}
+		}
+	}
+
+	for (URigVMNode* RenamedNode : RenamedNodes)
+	{
+		Graph->Notify(ERigVMGraphNotifType::VariableRenamed, Graph, RenamedNode);
+		Graph->MarkPackageDirty();
+	}
+
+	if (bUndo)
+	{
+		if (RenamedNodes.Num() > 0)
+		{
+			ActionStack->EndAction(Action);
+		}
+		else
+		{
+			ActionStack->CancelAction(Action);
+		}
+	}
+
+	return RenamedNodes.Num() > 0;
+}
+
+bool URigVMController::RenameParameter(const FName& InOldName, const FName& InNewName, bool bUndo)
+{
+	if(!IsValidGraph())
+	{
+		return false;
+	}
+
+	if (InOldName == InNewName)
+	{
+		ReportWarning(TEXT("RenameParameter: InOldName and InNewName are equal."));
+		return false;
+	}
+
+	TArray<FRigVMGraphParameterDescription> ExistingParameters = Graph->GetParameterDescriptions();
+	for (const FRigVMGraphParameterDescription& ExistingParameter : ExistingParameters)
+	{
+		if (ExistingParameter.Name == InNewName)
+		{
+			ReportErrorf(TEXT("Cannot rename parameter to '%s' - parameter already exists."), *InNewName.ToString());
+			return false;
+		}
+	}
+
+	FRigVMRenameParameterAction Action;
+	if (bUndo)
+	{
+		Action = FRigVMRenameParameterAction(InOldName, InNewName);
+		Action.Title = FString::Printf(TEXT("Rename Parameter"));
+		ActionStack->BeginAction(Action);
+	}
+
+	TArray<URigVMNode*> RenamedNodes;
+	for (URigVMNode* Node : Graph->Nodes)
+	{
+		if(URigVMParameterNode* ParameterNode = Cast<URigVMParameterNode>(Node))
+		{
+			if (ParameterNode->ParameterName == InOldName)
+			{
+				ParameterNode->ParameterName = InNewName;
+				RenamedNodes.Add(Node);
+			}
+		}
+	}
+
+	for (URigVMNode* RenamedNode : RenamedNodes)
+	{
+		Graph->Notify(ERigVMGraphNotifType::ParameterRenamed, Graph, RenamedNode);
+		Graph->MarkPackageDirty();
+	}
+
+	if (bUndo)
+	{
+		if (RenamedNodes.Num() > 0)
+		{
+			ActionStack->EndAction(Action);
+		}
+		else
+		{
+			ActionStack->CancelAction(Action);
+		}
+	}
+
+	return RenamedNodes.Num() > 0;
+}
+
+FString URigVMController::GetPinDefaultValue(const FString& InPinPath)
+{
+	if (!IsValidGraph())
+	{
+		return FString();
+	}
+
+	URigVMPin* Pin = Graph->FindPin(InPinPath);
+	if (Pin == nullptr)
+	{
+		ReportErrorf(TEXT("Cannot find pin '%s'."), *InPinPath);
+		return FString();
+	}
+
+	return Pin->GetDefaultValue();
+}
+
+bool URigVMController::SetPinDefaultValue(const FString& InPinPath, const FString& InDefaultValue, bool bResizeArrays, bool bUndo, bool bMergeUndoAction)
+{
+	if (!IsValidGraph())
+	{
+		return false;
+	}
+
+	URigVMPin* Pin = Graph->FindPin(InPinPath);
+	if (Pin == nullptr)
+	{
+		ReportErrorf(TEXT("Cannot find pin '%s'."), *InPinPath);
+		return false;
+	}
+
+	SetPinDefaultValue(Pin, InDefaultValue, bResizeArrays, bUndo, bMergeUndoAction);
+	return true;
+}
+
+void URigVMController::SetPinDefaultValue(URigVMPin* InPin, const FString& InDefaultValue, bool bResizeArrays, bool bUndo, bool bMergeUndoAction)
+{
+	check(InPin);
+	ensure(!InDefaultValue.IsEmpty());
+
+	FRigVMSetPinDefaultValueAction Action;
+	if (bUndo)
+	{
+		Action = FRigVMSetPinDefaultValueAction(InPin, InDefaultValue);
+		Action.Title = FString::Printf(TEXT("Set Pin Default Value"));
+		ActionStack->BeginAction(Action);
+	}
+
+	bool bSetPinDefaultValueSucceeded = false;
+	if (InPin->IsArray())
+	{
+		if (ShouldPinBeUnfolded(InPin))
+		{
+			TArray<FString> Elements = SplitDefaultValue(InDefaultValue);
+
+			if (bResizeArrays)
+			{
+				while (Elements.Num() > InPin->SubPins.Num())
+				{
+					InsertArrayPin(InPin, INDEX_NONE, FString(), bUndo);
+				}
+				while (Elements.Num() < InPin->SubPins.Num())
+				{
+					RemovePin(InPin->SubPins.Last(), bUndo);
+				}
+			}
+			else
+			{
+				ensure(Elements.Num() == InPin->SubPins.Num());
+			}
+
+			for (int32 ElementIndex = 0; ElementIndex < Elements.Num(); ElementIndex++)
+			{
+				URigVMPin* SubPin = InPin->SubPins[ElementIndex];
+				if (SubPin->IsStringType())
+				{
+					Elements[ElementIndex].Mid(1, Elements[ElementIndex].Len() - 2);
+				}
+				SetPinDefaultValue(SubPin, Elements[ElementIndex], bResizeArrays, false, false);
+				bSetPinDefaultValueSucceeded = true;
+			}
+		}
+	}
+	else if (InPin->IsStruct())
+	{
+		TArray<FString> MemberValuePairs = SplitDefaultValue(InDefaultValue);
+
+		for (const FString& MemberValuePair : MemberValuePairs)
+		{
+			FString MemberName, MemberValue;
+			if (MemberValuePair.Split(TEXT("="), &MemberName, &MemberValue))
+			{
+				URigVMPin* SubPin = InPin->FindSubPin(MemberName);
+				if (SubPin)
+				{
+					if (SubPin->IsStringType())
+					{
+						MemberValue.Mid(1, MemberValue.Len() - 2);
+					}
+					SetPinDefaultValue(SubPin, MemberValue, bResizeArrays, false, false);
+					bSetPinDefaultValueSucceeded = true;
+				}
+			}
+		}
+	}
+	
+	if(!bSetPinDefaultValueSucceeded)
+	{
+		ensure(InPin->GetSubPins().Num() == 0);;
+		InPin->DefaultValue = InDefaultValue;
+		Graph->Notify(ERigVMGraphNotifType::PinDefaultValueChanged, Graph, InPin);
+		Graph->MarkPackageDirty();
+	}
+
+	if (bUndo)
+	{
+		ActionStack->EndAction(Action, bMergeUndoAction);
+	}
+}
+
+TArray<FString> URigVMController::SplitDefaultValue(const FString& InDefaultValue)
+{
+	TArray<FString> Parts;
+	if (InDefaultValue.IsEmpty())
+	{
+		return Parts;
+	}
+
+	ensure(InDefaultValue[0] == TCHAR('('));
+	ensure(InDefaultValue[InDefaultValue.Len() - 1] == TCHAR(')'));
+
+	FString Content = InDefaultValue.Mid(1, InDefaultValue.Len() - 2);
+	int32 BraceCount = 0;
+	int32 QuoteCount = 0;
+
+	int32 LastPartStartIndex = 0;
+	for (int32 CharIndex = 0; CharIndex < Content.Len(); CharIndex++)
+	{
+		TCHAR Char = Content[CharIndex];
+		if (QuoteCount > 0)
+		{
+			if (Content[CharIndex - 1] == TCHAR('\\') && Char == TCHAR('"'))
+			{
+				QuoteCount = 0;
+			}
+		}
+		else if (Char == TCHAR('"'))
+		{
+			QuoteCount = 1;
+		}
+		else if (Char == TCHAR('('))
+		{
+			BraceCount++;
+		}
+		else if (Char == TCHAR(')'))
+		{
+			BraceCount--;
+			BraceCount = FMath::Max<int32>(BraceCount, 0);
+		}
+		else if (Char == TCHAR(',') && BraceCount == 0)
+		{
+			Parts.Add(Content.Mid(LastPartStartIndex, CharIndex - LastPartStartIndex));
+			LastPartStartIndex = CharIndex + 1;
+		}
+	}
+
+	Parts.Add(Content.Mid(LastPartStartIndex));
+	return Parts;
+}
+
+FString URigVMController::AddArrayPin(const FString& InArrayPinPath, const FString& InDefaultValue, bool bUndo)
+{
+	return InsertArrayPin(InArrayPinPath, INDEX_NONE, InDefaultValue, bUndo);
+}
+
+FString URigVMController::DuplicateArrayPin(const FString& InArrayElementPinPath, bool bUndo)
+{
+	if (!IsValidGraph())
+	{
+		return FString();
+	}
+
+	URigVMPin* ElementPin = Graph->FindPin(InArrayElementPinPath);
+	if (ElementPin == nullptr)
+	{
+		ReportErrorf(TEXT("Cannot find pin '%s'."), *InArrayElementPinPath);
+		return FString();
+	}
+
+	if (!ElementPin->IsArrayElement())
+	{
+		ReportErrorf(TEXT("Pin '%s' is not an array element."), *InArrayElementPinPath);
+		return FString();
+	}
+
+	URigVMPin* ArrayPin = ElementPin->GetParentPin();
+	check(ArrayPin);
+	ensure(ArrayPin->IsArray());
+
+	FString DefaultValue = ElementPin->GetDefaultValue();
+	return InsertArrayPin(ArrayPin->GetPinPath(), ElementPin->GetPinIndex() + 1, DefaultValue, bUndo);
+}
+
+FString URigVMController::InsertArrayPin(const FString& InArrayPinPath, int32 InIndex, const FString& InDefaultValue, bool bUndo)
+{
+	if (!IsValidGraph())
+	{
+		return FString();
+	}
+
+	URigVMPin* ArrayPin = Graph->FindPin(InArrayPinPath);
+	if (ArrayPin == nullptr)
+	{
+		ReportErrorf(TEXT("Cannot find pin '%s'."), *InArrayPinPath);
+		return FString();
+	}
+
+	URigVMPin* ElementPin = InsertArrayPin(ArrayPin, InIndex, InDefaultValue, bUndo);
+	if (ElementPin)
+	{
+		return ElementPin->GetPinPath();
+	}
+
+	return FString();
+}
+
+URigVMPin* URigVMController::InsertArrayPin(URigVMPin* ArrayPin, int32 InIndex, const FString& InDefaultValue, bool bUndo)
+{
+	if (!ArrayPin->IsArray())
+	{
+		ReportErrorf(TEXT("Pin '%s' is not an array."), *ArrayPin->GetPinPath());
+		return nullptr;
+	}
+
+	if (!ShouldPinBeUnfolded(ArrayPin))
+	{
+		ReportErrorf(TEXT("Cannot insert array pin under '%s'."), *ArrayPin->GetPinPath());
+		return nullptr;
+	}
+
+	if (InIndex == INDEX_NONE)
+	{
+		InIndex = ArrayPin->GetSubPins().Num();
+	}
+
+	FRigVMInsertArrayPinAction Action;
+	if (bUndo)
+	{
+		Action = FRigVMInsertArrayPinAction(ArrayPin, InIndex, InDefaultValue);
+		Action.Title = FString::Printf(TEXT("Insert Array Pin"));
+		ActionStack->BeginAction(Action);
+	}
+
+	for (int32 ExistingIndex = ArrayPin->GetSubPins().Num() - 1; ExistingIndex >= InIndex; ExistingIndex--)
+	{
+		URigVMPin* ExistingPin = ArrayPin->GetSubPins()[ExistingIndex];
+		ExistingPin->Rename(*FString::FormatAsNumber(ExistingIndex + 1));
+	}
+
+	URigVMPin* Pin = NewObject<URigVMPin>(ArrayPin, *FString::FormatAsNumber(InIndex));
+	ConfigurePinFromPin(Pin, ArrayPin);
+	Pin->CPPType = ArrayPin->GetArrayElementCppType();
+	ArrayPin->SubPins.Insert(Pin, InIndex);
+
+	if (Pin->IsStruct())
+	{
+		UScriptStruct* ScriptStruct = Pin->GetScriptStruct();
+		if (ScriptStruct)
+		{
+			AddPinsForStruct(ScriptStruct, Pin->GetNode(), Pin, Pin->Direction, InDefaultValue);
+		}
+	}
+	else if (Pin->IsArray())
+	{
+		UArrayProperty * ArrayProperty = Cast<UArrayProperty>(FindPropertyForPin(Pin->GetPinPath()));
+		if (ArrayProperty)
+		{
+			TArray<FString> ElementDefaultValues = SplitDefaultValue(InDefaultValue);
+			AddPinsForArray(ArrayProperty, Pin->GetNode(), Pin, Pin->Direction, ElementDefaultValues);
+		}
+	}
+	else
+	{
+		Pin->DefaultValue = InDefaultValue;
+	}
+
+	Graph->Notify(ERigVMGraphNotifType::PinArraySizeChanged, Graph, ArrayPin);
+	Graph->MarkPackageDirty();
+
+	if (bUndo)
+	{
+		ActionStack->EndAction(Action);
+	}
+
+	return Pin;
+}
+
+bool URigVMController::RemoveArrayPin(const FString& InArrayElementPinPath, bool bUndo)
+{
+	if(!IsValidGraph())
+	{
+		return false;
+	}
+
+	URigVMPin* ArrayElementPin = Graph->FindPin(InArrayElementPinPath);
+	if (ArrayElementPin == nullptr)
+	{
+		ReportErrorf(TEXT("Cannot find pin '%s'."), *InArrayElementPinPath);
+		return false;
+	}
+
+	if (!ArrayElementPin->IsArrayElement())
+	{
+		ReportErrorf(TEXT("Pin '%s' is not an array element."), *InArrayElementPinPath);
+		return false;
+	}
+
+	URigVMPin* ArrayPin = ArrayElementPin->GetParentPin();
+	check(ArrayPin);
+	ensure(ArrayPin->IsArray());
+
+	if (ArrayPin->GetArraySize() == 1)
+	{
+		ReportErrorf(TEXT("Cannot remove the last element from array pin '%s'."), *ArrayPin->GetPinPath());
+		return false;
+	}
+
+	FRigVMRemoveArrayPinAction Action;
+	if (bUndo)
+	{
+		Action = FRigVMRemoveArrayPinAction(ArrayElementPin);
+		Action.Title = FString::Printf(TEXT("Remove Array Pin"));
+		ActionStack->BeginAction(Action);
+	}
+
+	int32 IndexToRemove = ArrayElementPin->GetPinIndex();
+	if (!RemovePin(ArrayElementPin, bUndo))
+	{
+		return false;
+	}
+
+	for (int32 ExistingIndex = ArrayPin->GetSubPins().Num() - 1; ExistingIndex >= IndexToRemove; ExistingIndex--)
+	{
+		URigVMPin* ExistingPin = ArrayPin->GetSubPins()[ExistingIndex];
+		ExistingPin->SetNameFromIndex();
+	}
+
+	Graph->MarkPackageDirty();
+	Graph->Notify(ERigVMGraphNotifType::PinArraySizeChanged, Graph, ArrayPin);
+
+	if (bUndo)
+	{
+		ActionStack->EndAction(Action);
+	}
+
+	return true;
+}
+
+bool URigVMController::RemovePin(URigVMPin* InPinToRemove, bool bUndo)
+{
+	if (bUndo)
+	{
+		BreakAllLinks(InPinToRemove, true, bUndo);
+		BreakAllLinks(InPinToRemove, false, bUndo);
+		BreakAllLinksRecursive(InPinToRemove, true, false, bUndo);
+		BreakAllLinksRecursive(InPinToRemove, false, false, bUndo);
+	}
+
+	URigVMPin* ParentPin = InPinToRemove->GetParentPin();
+	ParentPin->SubPins.Remove(InPinToRemove);
+	InPinToRemove->Rename(*FString::Printf(TEXT("URigVMPin_%d_Removed"), ++sRemovedPinIndex), GetTransientPackage());
+
+	TArray<URigVMPin*> SubPins = InPinToRemove->GetSubPins();
+	for (URigVMPin* SubPin : SubPins)
+	{
+		if (!RemovePin(SubPin, bUndo))
+		{
+			return false;
+		}
+	}
+
+	InPinToRemove->MarkPendingKill();
+	return true;
+}
+
+bool URigVMController::ClearArrayPin(const FString& InArrayPinPath, bool bUndo)
+{
+	if(!IsValidGraph())
+	{
+		return false;
+	}
+
+	URigVMPin* Pin = Graph->FindPin(InArrayPinPath);
+	if (Pin == nullptr)
+	{
+		ReportErrorf(TEXT("Cannot find pin '%s'."), *InArrayPinPath);
+		return false;
+	}
+
+	if (!Pin->IsArray())
+	{
+		ReportErrorf(TEXT("Pin '%s' is not an array."), *InArrayPinPath);
+		return false;
+	}
+
+	TArray<URigVMPin*> ElementPins = Pin->GetSubPins();
+	if (ElementPins.Num() == 0)
+	{
+		return false;
+	}
+
+	FRigVMBaseAction Action;
+	if (bUndo)
+	{
+		Action.Title = FString::Printf(TEXT("Clear Array Pin"));
+		ActionStack->BeginAction(Action);
+	}
+
+	int32 RemovedPins = 0;
+	for (int32 ElementIndex = ElementPins.Num() - 1; ElementIndex > 0; ElementIndex--)
+	{
+		if (!RemoveArrayPin(ElementPins[ElementIndex]->GetPinPath(), bUndo))
+		{
+			if (bUndo)
+			{
+				ActionStack->CancelAction(Action);
+			}
+			return false;
+		}
+		RemovedPins++;
+	}
+
+	if (bUndo)
+	{
+		if (RemovedPins > 0)
+		{
+			ActionStack->EndAction(Action);
+		}
+		else
+		{
+			ActionStack->CancelAction(Action);
+		}
+	}
+
+	return RemovedPins > 0;
+}
+
+bool URigVMController::AddLink(const FString& InOutputPinPath, const FString& InInputPinPath, bool bUndo)
+{
+	if(!IsValidGraph())
+	{
+		return false;
+	}
+
+	URigVMPin* OutputPin = Graph->FindPin(InOutputPinPath);
+	if (OutputPin == nullptr)
+	{
+		ReportErrorf(TEXT("Cannot find pin '%s'."), *InOutputPinPath);
+		return false;
+	}
+	URigVMPin* InputPin = Graph->FindPin(InInputPinPath);
+	if (InputPin == nullptr)
+	{
+		ReportErrorf(TEXT("Cannot find pin '%s'."), *InInputPinPath);
+		return false;
+	}
+
+	return AddLink(OutputPin, InputPin, bUndo);
+}
+
+bool URigVMController::AddLink(URigVMPin* OutputPin, URigVMPin* InputPin, bool bUndo)
+{
+	if(OutputPin == nullptr)
+	{
+		ReportError(TEXT("OutputPin is nullptr."));
+		return false;
+	}
+
+	if(InputPin == nullptr)
+	{
+		ReportError(TEXT("InputPin is nullptr."));
+		return false;
+	}
+
+	if(!IsValidPinForGraph(OutputPin) || !IsValidPinForGraph(InputPin))
+	{
+		return false;
+	}
+
+	FString FailureReason;
+	if (!Graph->CanLink(OutputPin, InputPin, &FailureReason))
+	{
+		ReportErrorf(TEXT("Cannot link '%s' to '%s': %s."), *OutputPin->GetPinPath(), *InputPin->GetPinPath(), *FailureReason);
+		return false;
+	}
+
+	ensure(!OutputPin->IsLinkedTo(InputPin));
+	ensure(!InputPin->IsLinkedTo(OutputPin));
+
+	FRigVMAddLinkAction Action;
+	if (bUndo)
+	{
+		Action = FRigVMAddLinkAction(OutputPin, InputPin);
+		Action.Title = FString::Printf(TEXT("Add Link"));
+		ActionStack->BeginAction(Action);
+	}
+
+	BreakAllLinks(InputPin, true, bUndo);
+	if (bUndo)
+	{
+		BreakAllLinksRecursive(InputPin, true, true, bUndo);
+		BreakAllLinksRecursive(InputPin, true, false, bUndo);
+	}
+
+	URigVMLink* Link = NewObject<URigVMLink>(Graph);
+	Link->SourcePin = OutputPin;
+	Link->TargetPin = InputPin;
+	Link->SourcePinPath = OutputPin->GetPinPath();
+	Link->TargetPinPath = InputPin->GetPinPath();
+	Graph->Links.Add(Link);
+	OutputPin->Links.Add(Link);
+	InputPin->Links.Add(Link);
+
+	Graph->MarkPackageDirty();
+	Graph->Notify(ERigVMGraphNotifType::LinkAdded, Graph, Link);
+
+	if (bUndo)
+	{
+		ActionStack->EndAction(Action);
+	}
+
+	return true;
+}
+
+bool URigVMController::BreakLink(const FString& InOutputPinPath, const FString& InInputPinPath, bool bUndo)
+{
+	if(!IsValidGraph())
+	{
+		return false;
+	}
+
+	URigVMPin* OutputPin = Graph->FindPin(InOutputPinPath);
+	if (OutputPin == nullptr)
+	{
+		ReportErrorf(TEXT("Cannot find pin '%s'."), *InOutputPinPath);
+		return false;
+	}
+	URigVMPin* InputPin = Graph->FindPin(InInputPinPath);
+	if (InputPin == nullptr)
+	{
+		ReportErrorf(TEXT("Cannot find pin '%s'."), *InInputPinPath);
+		return false;
+	}
+
+	return BreakLink(OutputPin, InputPin, bUndo);
+}
+
+bool URigVMController::BreakLink(URigVMPin* OutputPin, URigVMPin* InputPin, bool bUndo)
+{
+	if(!IsValidPinForGraph(OutputPin) || !IsValidPinForGraph(InputPin))
+	{
+		return false;
+	}
+
+	if (!OutputPin->IsLinkedTo(InputPin))
+	{
+		return false;
+	}
+	ensure(InputPin->IsLinkedTo(OutputPin));
+
+	for (URigVMLink* Link : InputPin->Links)
+	{
+		if (Link->SourcePin == OutputPin && Link->TargetPin == InputPin)
+		{
+			FRigVMBreakLinkAction Action;
+			if (bUndo)
+			{
+				Action = FRigVMBreakLinkAction(OutputPin, InputPin);
+				Action.Title = FString::Printf(TEXT("Break Link"));
+				ActionStack->BeginAction(Action);
+			}
+
+			OutputPin->Links.Remove(Link);
+			InputPin->Links.Remove(Link);
+			Graph->Links.Remove(Link);
+			
+			Graph->MarkPackageDirty();
+			Graph->Notify(ERigVMGraphNotifType::LinkRemoved, Graph, Link);
+
+			Link->MarkPendingKill();
+
+			if (bUndo)
+			{
+				ActionStack->EndAction(Action);
+			}
+
+			return true;
+		}
+	}
+
+	return false;
+}
+
+bool URigVMController::BreakAllLinks(const FString& InPinPath, bool bAsInput, bool bUndo)
+{
+	if(!IsValidGraph())
+	{
+		return false;
+	}
+
+	URigVMPin* Pin = Graph->FindPin(InPinPath);
+	if (Pin == nullptr)
+	{
+		ReportErrorf(TEXT("Cannot find pin '%s'."), *InPinPath);
+		return false;
+	}
+
+	return BreakAllLinks(Pin, bAsInput, bUndo);
+}
+
+bool URigVMController::BreakAllLinks(URigVMPin* Pin, bool bAsInput, bool bUndo)
+{
+	if(!IsValidPinForGraph(Pin))
+	{
+		return false;
+	}
+
+	FRigVMBaseAction Action;
+	if (bUndo)
+	{
+		Action.Title = FString::Printf(TEXT("Break All Links"));
+		ActionStack->BeginAction(Action);
+	}
+
+	int32 LinksBroken = 0;
+	TArray<URigVMLink*> Links = Pin->GetLinks();
+	for (int32 LinkIndex = Links.Num() - 1; LinkIndex >= 0; LinkIndex--)
+	{
+		URigVMLink* Link = Links[LinkIndex];
+		if (bAsInput && Link->GetTargetPin() == Pin)
+		{
+			LinksBroken += BreakLink(Link->GetSourcePin(), Pin, bUndo) ? 1 : 0;
+		}
+		else if (!bAsInput && Link->GetSourcePin() == Pin)
+		{
+			LinksBroken += BreakLink(Pin, Link->GetTargetPin(), bUndo) ? 1 : 0;
+		}
+	}
+
+	if (bUndo)
+	{
+		if (LinksBroken > 0)
+		{
+			ActionStack->EndAction(Action);
+		}
+		else
+		{
+			ActionStack->CancelAction(Action);
+		}
+	}
+
+	return LinksBroken > 0;
+}
+
+void URigVMController::BreakAllLinksRecursive(URigVMPin* Pin, bool bAsInput, bool bTowardsParent, bool bUndo)
+{
+	if (bTowardsParent)
+	{
+		URigVMPin* ParentPin = Pin->GetParentPin();
+		if (ParentPin)
+		{
+			BreakAllLinks(ParentPin, bAsInput, bUndo);
+			BreakAllLinksRecursive(ParentPin, bAsInput, bTowardsParent, bUndo);
+		}
+	}
+	else
+	{
+		for (URigVMPin* SubPin : Pin->SubPins)
+		{
+			BreakAllLinks(SubPin, bAsInput, bUndo);
+			BreakAllLinksRecursive(SubPin, bAsInput, bTowardsParent, bUndo);
+		}
+	}
+}
+
+FString URigVMController::GetValidName(const FString& InPrefix)
+{
+	int32 NameSuffix = 0;
+	FString Name = InPrefix;
+
+	while (!Graph->IsNameAvailable(Name))
+	{
+		Name = FString::Printf(TEXT("%s_%d"), *InPrefix, ++NameSuffix);
+	}
+
+	return Name;
+}
+
+bool URigVMController::IsValidGraph()
+{
+	if (Graph == nullptr)
+	{
+		ReportError(TEXT("Controller does not have a graph associated - use SetGraph / set_graph."));
+		return false;
+	}
+	return true;
+}
+
+bool URigVMController::IsValidNodeForGraph(URigVMNode* InNode)
+{
+	if(!IsValidGraph())
+	{
+		return false;
+	}
+
+	if (InNode == nullptr)
+	{
+		ReportError(TEXT("InNode is nullptr."));
+		return false;
+	}
+
+	if (InNode->GetGraph() != Graph)
+	{
+		ReportErrorf(TEXT("InNode '%s' is on a different graph."), *InNode->GetNodePath());
+		return false;
+	}
+
+	if (InNode->GetNodeIndex() == INDEX_NONE)
+	{
+		ReportErrorf(TEXT("InNode '%s' is transient (not yet nested to a graph)."), *InNode->GetName());
+	}
+
+	return true;
+}
+
+bool URigVMController::IsValidPinForGraph(URigVMPin* InPin)
+{
+	if(!IsValidGraph())
+	{
+		return false;
+	}
+
+	if (InPin == nullptr)
+	{
+		ReportError(TEXT("InPin is nullptr."));
+		return false;
+	}
+
+	if (!IsValidNodeForGraph(InPin->GetNode()))
+	{
+		return false;
+	}
+
+	if (InPin->GetPinIndex() == INDEX_NONE)
+	{
+		ReportErrorf(TEXT("InPin '%s' is transient (not yet nested properly)."), *InPin->GetName());
+	}
+
+	return true;
+}
+
+bool URigVMController::IsValidLinkForGraph(URigVMLink* InLink)
+{
+	if(!IsValidGraph())
+	{
+		return false;
+	}
+
+	if (InLink == nullptr)
+	{
+		ReportError(TEXT("InLink is nullptr."));
+		return false;
+	}
+
+	if (InLink->GetGraph() != Graph)
+	{
+		ReportError(TEXT("InLink is on a different graph."));
+		return false;
+	}
+
+	if(InLink->GetSourcePin() == nullptr)
+	{
+		ReportError(TEXT("InLink has no source pin."));
+		return false;
+	}
+
+	if(InLink->GetTargetPin() == nullptr)
+	{
+		ReportError(TEXT("InLink has no target pin."));
+		return false;
+	}
+
+	if (InLink->GetLinkIndex() == INDEX_NONE)
+	{
+		ReportError(TEXT("InLink is transient (not yet nested properly)."));
+	}
+
+	if(!IsValidPinForGraph(InLink->GetSourcePin()))
+	{
+		return false;
+	}
+
+	if(!IsValidPinForGraph(InLink->GetTargetPin()))
+	{
+		return false;
+	}
+
+	return true;
+}
+
+void URigVMController::AddPinsForStruct(UStruct* InStruct, URigVMNode* InNode, URigVMPin* InParentPin, ERigVMPinDirection InPinDirection, const FString& InDefaultValue)
+{
+	TArray<FString> MemberNameValuePairs = SplitDefaultValue(InDefaultValue);
+	TMap<FName, FString> MemberValues;
+	for (const FString& MemberNameValuePair : MemberNameValuePairs)
+	{
+		FString MemberName, MemberValue;
+		if (MemberNameValuePair.Split(TEXT("="), &MemberName, &MemberValue))
+		{
+			MemberValues.Add(*MemberName, MemberValue);
+		}
+	}
+
+	for (TFieldIterator<UProperty> It(InStruct); It; ++It)
+	{
+		FName PropertyName = It->GetFName();
+
+		URigVMPin* Pin = NewObject<URigVMPin>(InParentPin == nullptr ? Cast<UObject>(InNode) : Cast<UObject>(InParentPin), PropertyName);
+		ConfigurePinFromProperty(*It, Pin, InPinDirection);
+
+		if (InParentPin)
+		{
+			InParentPin->SubPins.Add(Pin);
+		}
+		else
+		{
+			InNode->Pins.Add(Pin);
+		}
+
+		FString* DefaultValuePtr = MemberValues.Find(Pin->GetFName());
+
+		UStructProperty* StructProperty = Cast<UStructProperty>(*It);
+		if (StructProperty)
+		{
+			if (ShouldStructBeUnfolded(StructProperty->Struct))
+			{
+				FString DefaultValue;
+				if (DefaultValuePtr != nullptr)
+				{
+					DefaultValue = *DefaultValuePtr;
+				}
+				AddPinsForStruct(StructProperty->Struct, InNode, Pin, Pin->GetDirection(), DefaultValue);
+			}
+			else if(DefaultValuePtr != nullptr)
+			{
+				Pin->DefaultValue = *DefaultValuePtr;
+			}
+		}
+
+		UArrayProperty* ArrayProperty = Cast<UArrayProperty>(*It);
+		if (ArrayProperty)
+		{
+			ensure(Pin->IsArray());
+
+			if (DefaultValuePtr)
+			{
+				if (ShouldPinBeUnfolded(Pin))
+				{
+					TArray<FString> ElementDefaultValues = SplitDefaultValue(*DefaultValuePtr);
+					AddPinsForArray(ArrayProperty, InNode, Pin, Pin->Direction, ElementDefaultValues);
+				}
+				else
+				{
+					Pin->DefaultValue = *DefaultValuePtr;
+				}
+			}
+		}
+		
+		if (!Pin->IsArray() && !Pin->IsStruct() && DefaultValuePtr != nullptr)
+		{
+			FString DefaultValue = *DefaultValuePtr;
+			if (Pin->IsStringType())
+			{
+				DefaultValue = DefaultValue.Mid(1, DefaultValue.Len() - 2);
+			}
+			Pin->DefaultValue = DefaultValue;
+		}
+	}
+}
+
+void URigVMController::AddPinsForArray(UArrayProperty* InArrayProperty, URigVMNode* InNode, URigVMPin* InParentPin, ERigVMPinDirection InPinDirection, const TArray<FString>& InDefaultValues)
+{
+	check(InParentPin);
+	if (!ShouldPinBeUnfolded(InParentPin))
+	{
+		return;
+	}
+
+	for (int32 ElementIndex = 0; ElementIndex < InDefaultValues.Num(); ElementIndex++)
+	{
+		FString ElementName = FString::FormatAsNumber(InParentPin->SubPins.Num());
+		URigVMPin* Pin = NewObject<URigVMPin>(InParentPin, *ElementName);
+		ConfigurePinFromProperty(InArrayProperty->Inner, Pin, InPinDirection);
+		FString DefaultValue = InDefaultValues[ElementIndex];
+
+		InParentPin->SubPins.Add(Pin);
+
+		UStructProperty* StructProperty = Cast<UStructProperty>(InArrayProperty->Inner);
+		if (StructProperty)
+		{
+			if (ShouldPinBeUnfolded(Pin))
+			{
+				AddPinsForStruct(StructProperty->Struct, InNode, Pin, Pin->Direction, DefaultValue);
+			}
+			else if(!DefaultValue.IsEmpty())
+			{
+				Pin->DefaultValue = DefaultValue;
+			}
+		}
+
+		UArrayProperty* ArrayProperty = Cast<UArrayProperty>(InArrayProperty->Inner);
+		if (ArrayProperty)
+		{
+			if (ShouldPinBeUnfolded(Pin))
+			{
+				TArray<FString> ElementDefaultValues = SplitDefaultValue(DefaultValue);
+				AddPinsForArray(ArrayProperty, InNode, Pin, Pin->Direction, ElementDefaultValues);
+			}
+			else if (!DefaultValue.IsEmpty())
+			{
+				Pin->DefaultValue = DefaultValue;
+			}
+		}
+
+		if (!Pin->IsArray() && !Pin->IsStruct())
+		{
+			if (Pin->IsStringType())
+			{
+				DefaultValue = DefaultValue.Mid(1, DefaultValue.Len() - 2);
+			}
+			Pin->DefaultValue = DefaultValue;
+		}
+	}
+}
+
+void URigVMController::ConfigurePinFromProperty(UProperty* InProperty, URigVMPin* InOutPin, ERigVMPinDirection InPinDirection)
+{
+	if (InPinDirection == ERigVMPinDirection::Invalid)
+	{
+#if WITH_EDITOR
+		bool bIsInput = InProperty->HasMetaData(TEXT("Input"));
+		bool bIsOutput = InProperty->HasMetaData(TEXT("Output"));
+
+		if (bIsInput)
+		{
+			InOutPin->Direction = bIsOutput ? ERigVMPinDirection::IO : ERigVMPinDirection::Input;
+		}
+		else
+		{
+			InOutPin->Direction = bIsOutput ? ERigVMPinDirection::Output : ERigVMPinDirection::Hidden;
+		}
+#endif
+	}
+	else
+	{
+		InOutPin->Direction = InPinDirection;
+	}
+
+#if WITH_EDITOR
+	InOutPin->bIsConstant = InProperty->HasMetaData(TEXT("Constant"));
+	FString CustomWidgetName = InProperty->GetMetaData(TEXT("Widget"));
+	InOutPin->CustomWidgetName = CustomWidgetName.IsEmpty() ? FName(NAME_None) : FName(*CustomWidgetName);
+#endif
+
+	FString ExtendedCppType;
+	InOutPin->CPPType = InProperty->GetCPPType(&ExtendedCppType);
+	InOutPin->CPPType += ExtendedCppType;
+
+	UStructProperty* StructProperty = Cast<UStructProperty>(InProperty);
+	if (StructProperty)
+	{
+		InOutPin->ScriptStruct = StructProperty->Struct;
+		InOutPin->ScriptStructPath = *InOutPin->ScriptStruct->GetPathName();
+	}
+}
+
+void URigVMController::ConfigurePinFromPin(URigVMPin* InOutPin, URigVMPin* InPin)
+{
+	InOutPin->bIsConstant = InPin->bIsConstant;
+	InOutPin->Direction = InPin->Direction;
+	InOutPin->CPPType = InPin->CPPType;
+	InOutPin->ScriptStructPath = InPin->ScriptStructPath;
+	InOutPin->ScriptStruct = InPin->ScriptStruct;
+	InOutPin->DefaultValue = InPin->DefaultValue;
+}
+
+bool URigVMController::ShouldStructBeUnfolded(const UStruct* Struct)
+{
+	if (Struct == nullptr)
+	{
+		return false;
+	}
+	if (Struct->IsChildOf(UClass::StaticClass()))
+	{
+		return false;
+	}
+	if(Struct->IsChildOf(FRigVMExecuteContext::StaticStruct()))
+	{
+		return false;
+	}
+	return true;
+}
+
+bool URigVMController::ShouldPinBeUnfolded(URigVMPin* InPin)
+{
+	if (InPin->IsStruct())
+	{
+		return ShouldStructBeUnfolded(InPin->GetScriptStruct());
+	}
+	else if (InPin->IsArray())
+	{
+		return InPin->GetDirection() == ERigVMPinDirection::Input ||
+			InPin->GetDirection() == ERigVMPinDirection::IO;
+	}
+	return false;
+}
+
+UProperty* URigVMController::FindPropertyForPin(const FString& InPinPath)
+{
+	if(!IsValidGraph())
+	{
+		return nullptr;
+	}
+
+	TArray<FString> Parts;
+	if (!URigVMPin::SplitPinPath(InPinPath, Parts))
+	{
+		return nullptr;
+	}
+
+	URigVMPin* Pin = Graph->FindPin(InPinPath);
+	if (Pin == nullptr)
+	{
+		ReportErrorf(TEXT("Cannot find pin '%s'."), *InPinPath);
+		return nullptr;
+	}
+
+	URigVMNode* Node = Pin->GetNode();
+
+	URigVMStructNode* StructNode = Cast<URigVMStructNode>(Node);
+	if (StructNode)
+	{
+		int32 PartIndex = 1; // cut off the first one since it's the node
+
+		UStruct* Struct = StructNode->ScriptStruct;
+		UProperty* Property = Struct->FindPropertyByName(*Parts[PartIndex++]);
+
+		while (PartIndex < Parts.Num() && Property != nullptr)
+		{
+			if (UArrayProperty* ArrayProperty = Cast<UArrayProperty>(Property))
+			{
+				Property = ArrayProperty->Inner;
+				PartIndex++;
+				continue;
+			}
+
+			if (UStructProperty* StructProperty = Cast<UStructProperty>(Property))
+			{
+				Struct = StructProperty->Struct;
+				Property = Struct->FindPropertyByName(*Parts[PartIndex++]);
+				continue;
+			}
+
+			break;
+		}
+
+		if (PartIndex == Parts.Num())
+		{
+			return Property;
+		}
+	}
+
+	return nullptr;
+}
+
+void URigVMController::DetachLinksFromPinObjects()
+{
+	check(Graph);
+
+	for (URigVMLink* Link : Graph->Links)
+	{
+		URigVMPin* SourcePin = Link->GetSourcePin();
+		URigVMPin* TargetPin = Link->GetTargetPin();
+
+		Link->SourcePinPath = SourcePin->GetPinPath();
+		Link->TargetPinPath = TargetPin->GetPinPath();
+		Link->SourcePin = nullptr;
+		Link->TargetPin = nullptr;
+	}
+}
+
+void URigVMController::ReattachLinksToPinObjects()
+{
+	check(Graph);
+
+	// fix up the pin links based on the persisted data
+	TArray<URigVMLink*> NewLinks;
+	for (URigVMLink* Link : Graph->Links)
+	{
+		URigVMPin* SourcePin = Link->GetSourcePin();
+		URigVMPin* TargetPin = Link->GetTargetPin();
+		if (SourcePin == nullptr)
+		{
+			if (TargetPin != nullptr)
+			{
+				TargetPin->Links.Remove(Link);
+			}
+			continue;
+		}
+		if (TargetPin == nullptr)
+		{
+			if (SourcePin != nullptr)
+			{
+				SourcePin->Links.Remove(Link);
+			}
+			continue;
+		}
+
+		SourcePin->Links.AddUnique(Link);
+		TargetPin->Links.AddUnique(Link);
+		NewLinks.Add(Link);
+	}
+	Graph->Links = NewLinks;
+}
+
+#if WITH_EDITOR
+
+void URigVMController::RepopulatePinsOnNode(URigVMNode* InNode)
+{
+	if (InNode == nullptr)
+	{
+		ReportError(TEXT("InNode is nullptr."));
+		return;
+	}
+
+	check(Graph);
+
+	TMap<FName, FString> DefaultValues;
+	for (URigVMPin* Pin : InNode->Pins)
+	{
+		FString DefaultValue = Pin->GetDefaultValue();
+		if (!DefaultValue.IsEmpty())
+		{
+			DefaultValues.Add(Pin->GetFName(), DefaultValue);
+		}
+	}
+
+	if (URigVMStructNode* StructNode = Cast<URigVMStructNode>(InNode))
+	{
+		StructNode->Pins.Reset();
+
+		UScriptStruct* ScriptStruct = StructNode->GetScriptStruct();
+
+		FString NodeColorMetadata;
+		ScriptStruct->GetStringMetaDataHierarchical(*URigVMNode::NodeColorName, &NodeColorMetadata);
+		if(!NodeColorMetadata.IsEmpty())
+		{
+			StructNode->NodeColor = GetColorFromMetadata(NodeColorMetadata);
+		}
+
+		TArray<uint8> TempBuffer;
+		TempBuffer.AddUninitialized(ScriptStruct->GetStructureSize());
+		ScriptStruct->InitializeDefaultValue(TempBuffer.GetData());
+
+		FString ExportedDefaultValue;
+		ScriptStruct->ExportText(ExportedDefaultValue, TempBuffer.GetData(), nullptr, nullptr, PPF_None, nullptr);
+		ScriptStruct->DestroyStruct(TempBuffer.GetData());
+		AddPinsForStruct(ScriptStruct, StructNode, nullptr, ERigVMPinDirection::Invalid, ExportedDefaultValue);
+	}
+	else
+	{
+		return;
+	}
+
+	for (URigVMPin* Pin : InNode->Pins)
+	{
+		FString* DefaultValue = DefaultValues.Find(Pin->GetFName());
+		if (DefaultValue)
+		{
+			SetPinDefaultValue(Pin, *DefaultValue, true, false, false);
+		}
+	}
+}
+
+#endif
+
+FLinearColor URigVMController::GetColorFromMetadata(const FString& InMetadata)
+{
+	FLinearColor Color = FLinearColor::Black;
+
+	FString Metadata = InMetadata;
+	Metadata.TrimStartAndEnd();
+	FString SplitString(TEXT(" "));
+	FString Red, Green, Blue, GreenAndBlue;
+	if (Metadata.Split(SplitString, &Red, &GreenAndBlue))
+	{
+		Red.TrimEnd();
+		GreenAndBlue.TrimStart();
+		if (GreenAndBlue.Split(SplitString, &Green, &Blue))
+		{
+			Green.TrimEnd();
+			Blue.TrimStart();
+
+			float RedValue = FCString::Atof(*Red);
+			float GreenValue = FCString::Atof(*Green);
+			float BlueValue = FCString::Atof(*Blue);
+			Color = FLinearColor(RedValue, GreenValue, BlueValue);
+		}
+	}
+
+	return Color;
+}
+
+void URigVMController::ReportWarning(const FString& InMessage)
+{
+	if(!bReportWarningsAndErrors)
+	{
+		return;
+	}
+	FScriptExceptionHandler::Get().HandleException(ELogVerbosity::Warning, *InMessage, *FString());
+}
+
+void URigVMController::ReportError(const FString& InMessage)
+{
+	if(!bReportWarningsAndErrors)
+	{
+		return;
+	}
+	FScriptExceptionHandler::Get().HandleException(ELogVerbosity::Error, *InMessage, *FString());
+}
diff --git a/Engine/Source/Editor/RigVMEditor/Private/RigVMModel/RigVMControllerActions.cpp b/Engine/Source/Editor/RigVMEditor/Private/RigVMModel/RigVMControllerActions.cpp
new file mode 100644
--- /dev/null
+++ b/Engine/Source/Editor/RigVMEditor/Private/RigVMModel/RigVMControllerActions.cpp
@@ -0,0 +1,863 @@
+// Copyright 1998-2019 Epic Games, Inc. All Rights Reserved.
+
+#include "RigVMModel/RigVMControllerActions.h"
+#include "UObject/Package.h"
+
+#if WITH_EDITOR
+#include "Misc/ITransaction.h"
+#endif
+
+FRigVMActionWrapper::FRigVMActionWrapper(const FRigVMActionKey& Key)
+{
+	ScriptStruct = FindObjectChecked<UScriptStruct>(ANY_PACKAGE, *Key.ScriptStructPath);
+	Data.SetNumUninitialized(ScriptStruct->GetStructureSize());
+	ScriptStruct->InitializeStruct(Data.GetData(), 1);
+	ScriptStruct->ImportText(*Key.ExportedText, Data.GetData(), nullptr, PPF_None, nullptr, ScriptStruct->GetName());
+}
+FRigVMActionWrapper::~FRigVMActionWrapper()
+{
+	if (Data.Num() > 0 && ScriptStruct != nullptr)
+	{
+		ScriptStruct->DestroyStruct(Data.GetData(), 1);
+	}
+}
+
+FRigVMBaseAction* FRigVMActionWrapper::GetAction()
+{
+	return (FRigVMBaseAction*)Data.GetData();
+}
+
+FString FRigVMActionWrapper::ExportText()
+{
+	FString ExportedText;
+	if (Data.Num() > 0 && ScriptStruct != nullptr)
+	{
+		ScriptStruct->ExportText(ExportedText, GetAction(), nullptr, nullptr, PPF_None, nullptr);
+	}
+	return ExportedText;
+}
+
+bool URigVMActionStack::OpenUndoBracket(const FString& InTitle)
+{
+	FRigVMBaseAction* Action = new FRigVMBaseAction;
+	Action->Title = InTitle;
+	BracketActions.Add(Action);
+	BeginAction(*Action);
+	return true;
+}
+
+bool URigVMActionStack::CloseUndoBracket()
+{
+	ensure(BracketActions.Num() > 0);
+	FRigVMBaseAction* Action = BracketActions.Pop();
+	EndAction(*Action);
+	delete(Action);
+	return true;
+}
+
+bool URigVMActionStack::Undo(URigVMController* InController)
+{
+	check(InController)
+
+	if (UndoActions.Num() == 0)
+	{
+		InController->ReportWarning(TEXT("Nothing to undo."));
+		return false;
+	}
+
+	FRigVMActionKey KeyToUndo = UndoActions.Pop();
+	FRigVMActionWrapper Wrapper(KeyToUndo);
+	if (Wrapper.GetAction()->Undo(InController))
+	{
+		RedoActions.Add(KeyToUndo);
+		return true;
+	}
+	return false;
+}
+
+bool URigVMActionStack::Redo(URigVMController* InController)
+{
+	check(InController)
+
+	if (RedoActions.Num() == 0)
+	{
+		InController->ReportWarning(TEXT("Nothing to redo."));
+		return false;
+	}
+
+	FRigVMActionKey KeyToRedo = RedoActions.Pop();
+	FRigVMActionWrapper Wrapper(KeyToRedo);
+	if (Wrapper.GetAction()->Redo(InController))
+	{
+		UndoActions.Add(KeyToRedo);
+		return true;
+	}
+	return false;
+}
+
+#if WITH_EDITOR
+
+void URigVMActionStack::PostTransacted(const FTransactionObjectEvent& TransactionEvent)
+{
+	if (TransactionEvent.GetEventType() == ETransactionObjectEventType::UndoRedo)
+	{
+		URigVMController* Controller = Cast< URigVMController>(GetOuter());
+		if (Controller == nullptr)
+		{
+			return;
+		}
+
+		while (ActionIndex < UndoActions.Num())
+		{
+			if (UndoActions.Num() == 0)
+			{
+				break;
+			}
+			if (!Undo(Controller))
+			{
+				return;
+			}
+		}
+		while (ActionIndex > UndoActions.Num())
+		{
+			if (RedoActions.Num() == 0)
+			{
+				break;
+			}
+			if (!Redo(Controller))
+			{
+				return;
+			}
+		}
+	}
+}
+
+#endif
+
+
+bool FRigVMBaseAction::Merge(const FRigVMBaseAction* Other)
+{
+	return SubActions.Num() == 0 && Other->SubActions.Num() == 0;
+}
+
+bool FRigVMBaseAction::Undo(URigVMController* InController)
+{
+	bool Result = true;
+	for (int32 KeyIndex = SubActions.Num() - 1; KeyIndex >= 0; KeyIndex--)
+	{
+		FRigVMActionWrapper Wrapper(SubActions[KeyIndex]);
+		if(!Wrapper.GetAction()->Undo(InController))
+		{
+			Result = false;
+		}
+	}
+	return Result;
+}
+
+bool FRigVMBaseAction::Redo(URigVMController* InController)
+{
+	bool Result = true;
+	for (int32 KeyIndex = 0; KeyIndex < SubActions.Num(); KeyIndex++)
+	{
+		FRigVMActionWrapper Wrapper(SubActions[KeyIndex]);
+		if (!Wrapper.GetAction()->Redo(InController))
+		{
+			Result = false;
+		}
+	}
+	return Result;
+}
+
+bool FRigVMInverseAction::Undo(URigVMController* InController)
+{
+	return FRigVMBaseAction::Redo(InController);
+}
+
+bool FRigVMInverseAction::Redo(URigVMController* InController)
+{
+	return FRigVMBaseAction::Undo(InController);
+}
+
+FRigVMAddStructNodeAction::FRigVMAddStructNodeAction()
+	: ScriptStructPath()
+	, MethodName(NAME_None)
+	, Position(FVector2D::ZeroVector)
+	, NodePath()
+{
+}
+
+FRigVMAddStructNodeAction::FRigVMAddStructNodeAction(URigVMStructNode* InNode)
+	: ScriptStructPath(InNode->GetScriptStruct()->GetPathName())
+	, MethodName(InNode->GetMethodName())
+	, Position(InNode->GetPosition())
+	, NodePath(InNode->GetNodePath())
+{
+}
+
+bool FRigVMAddStructNodeAction::Undo(URigVMController* InController)
+{
+	if (!FRigVMBaseAction::Undo(InController))
+	{
+		return false;
+	}
+	return InController->RemoveNodeByName(*NodePath, false);
+}
+
+bool FRigVMAddStructNodeAction::Redo(URigVMController* InController)
+{
+#if WITH_EDITOR
+	if (URigVMStructNode* Node = InController->AddStructNodeFromStructPath(ScriptStructPath, MethodName, Position, NodePath, false))
+	{
+		return FRigVMBaseAction::Redo(InController);
+	}
+#endif
+	return false;
+}
+
+FRigVMAddVariableNodeAction::FRigVMAddVariableNodeAction()
+	: VariableName(NAME_None)
+	, CPPType()
+	, ScriptStructPath()
+	, bIsGetter(false)
+	, DefaultValue()
+	, Position(FVector2D::ZeroVector)
+	, NodePath()
+{
+}
+
+FRigVMAddVariableNodeAction::FRigVMAddVariableNodeAction(URigVMVariableNode* InNode)
+	: VariableName(InNode->GetVariableName())
+	, CPPType(InNode->GetCPPType())
+	, ScriptStructPath(InNode->GetScriptStruct()->GetPathName())
+	, bIsGetter(InNode->IsGetter())
+	, DefaultValue(InNode->GetDefaultValue())
+	, Position(InNode->GetPosition())
+	, NodePath(InNode->GetNodePath())
+{
+}
+
+bool FRigVMAddVariableNodeAction::Undo(URigVMController* InController)
+{
+	if (!FRigVMBaseAction::Undo(InController))
+	{
+		return false;
+	}
+	return InController->RemoveNodeByName(*NodePath, false);
+}
+
+bool FRigVMAddVariableNodeAction::Redo(URigVMController* InController)
+{
+#if WITH_EDITOR
+	if (URigVMVariableNode* Node = InController->AddVariableNodeFromStructPath(VariableName, CPPType, ScriptStructPath, bIsGetter, DefaultValue, Position, NodePath, false))
+	{
+		return FRigVMBaseAction::Redo(InController);
+	}
+#endif
+	return false;
+}
+
+FRigVMAddParameterNodeAction::FRigVMAddParameterNodeAction()
+	: ParameterName(NAME_None)
+	, CPPType()
+	, ScriptStructPath()
+	, bIsInput(false)
+	, DefaultValue()
+	, Position(FVector2D::ZeroVector)
+	, NodePath()
+{
+}
+
+FRigVMAddParameterNodeAction::FRigVMAddParameterNodeAction(URigVMParameterNode* InNode)
+	: ParameterName(InNode->GetParameterName())
+	, CPPType(InNode->GetCPPType())
+	, ScriptStructPath(InNode->GetScriptStruct()->GetPathName())
+	, bIsInput(InNode->IsInput())
+	, DefaultValue(InNode->GetDefaultValue())
+	, Position(InNode->GetPosition())
+	, NodePath(InNode->GetNodePath())
+{
+}
+
+bool FRigVMAddParameterNodeAction::Undo(URigVMController* InController)
+{
+	if (!FRigVMBaseAction::Undo(InController))
+	{
+		return false;
+	}
+	return InController->RemoveNodeByName(*NodePath, false);
+}
+
+bool FRigVMAddParameterNodeAction::Redo(URigVMController* InController)
+{
+#if WITH_EDITOR
+	if (URigVMParameterNode* Node = InController->AddParameterNodeFromStructPath(ParameterName, CPPType, ScriptStructPath, bIsInput, DefaultValue, Position, NodePath, false))
+	{
+		return FRigVMBaseAction::Redo(InController);
+	}
+#endif
+	return false;
+}
+
+FRigVMAddCommentNodeAction::FRigVMAddCommentNodeAction()
+	: CommentText()
+	, Position(FVector2D::ZeroVector)
+	, Size(FVector2D::ZeroVector)
+	, Color(FLinearColor::Black)
+	, NodePath()
+{
+}
+
+FRigVMAddCommentNodeAction::FRigVMAddCommentNodeAction(URigVMCommentNode* InNode)
+	: CommentText(InNode->GetCommentText())
+	, Position(InNode->GetPosition())
+	, Size(InNode->GetSize())
+	, Color(InNode->GetNodeColor())
+	, NodePath(InNode->GetNodePath())
+{
+}
+
+bool FRigVMAddCommentNodeAction::Undo(URigVMController* InController)
+{
+	if (!FRigVMBaseAction::Undo(InController))
+	{
+		return false;
+	}
+	return InController->RemoveNodeByName(*NodePath, false);
+}
+
+bool FRigVMAddCommentNodeAction::Redo(URigVMController* InController)
+{
+#if WITH_EDITOR
+	if (URigVMCommentNode* Node = InController->AddCommentNode(CommentText, Position, Size, Color, NodePath, false))
+	{
+		return FRigVMBaseAction::Redo(InController);
+	}
+#endif
+	return false;
+}
+
+FRigVMAddRerouteNodeAction::FRigVMAddRerouteNodeAction()
+	: PinPath()
+	, bShowAsFullNode(false)
+	, bAsInput(false)
+	, Position(FVector2D::ZeroVector)
+	, NodePath()
+{
+}
+
+FRigVMAddRerouteNodeAction::FRigVMAddRerouteNodeAction(URigVMRerouteNode* InNode, const FString& InPinPath, bool bInAsInput)
+	: PinPath(InPinPath)
+	, bShowAsFullNode(InNode->GetShowsAsFullNode())
+	, bAsInput(bInAsInput)
+	, Position(InNode->GetPosition())
+	, NodePath(InNode->GetNodePath())
+{
+}
+
+FRigVMAddRerouteNodeAction::FRigVMAddRerouteNodeAction(URigVMRerouteNode* InNode)
+	: PinPath()
+	, bShowAsFullNode(InNode->GetShowsAsFullNode())
+	, bAsInput(false)
+	, Position(InNode->GetPosition())
+	, NodePath(InNode->GetNodePath())
+{
+	URigVMPin* ValuePin = InNode->FindPin(TEXT("Value"));
+	check(ValuePin)
+	ensure(ValuePin->GetLinks().Num() > 0);
+
+	for (URigVMLink* Link : ValuePin->GetLinks())
+	{
+		if (Link->GetSourcePin() == ValuePin)
+		{
+			PinPath = Link->GetTargetPin()->GetPinPath();
+			bAsInput = true;
+			break;
+		}
+	}
+}
+
+
+bool FRigVMAddRerouteNodeAction::Undo(URigVMController* InController)
+{
+	if (!FRigVMBaseAction::Undo(InController))
+	{
+		return false;
+	}
+	return InController->RemoveNodeByName(*NodePath, false);
+}
+
+bool FRigVMAddRerouteNodeAction::Redo(URigVMController* InController)
+{
+#if WITH_EDITOR
+	if (URigVMRerouteNode* Node = InController->AddRerouteNodeOnPin(PinPath, bShowAsFullNode, bAsInput, Position, NodePath, false))
+	{
+		return FRigVMBaseAction::Redo(InController);
+	}
+#endif
+	return false;
+}
+
+FRigVMRemoveNodeAction::FRigVMRemoveNodeAction(URigVMNode* InNode)
+{
+	FRigVMInverseAction InverseAction;
+
+	if (URigVMStructNode* StructNode = Cast<URigVMStructNode>(InNode))
+	{
+		InverseAction.AddAction(FRigVMAddStructNodeAction(StructNode));
+		for (URigVMPin* Pin : StructNode->GetPins())
+		{
+			if (Pin->GetDirection() == ERigVMPinDirection::Input ||
+				Pin->GetDirection() == ERigVMPinDirection::Visible)
+			{
+				InverseAction.AddAction(FRigVMSetPinDefaultValueAction(Pin, Pin->GetDefaultValue()));
+			}
+		}
+	}
+	else if (URigVMVariableNode* VariableNode = Cast<URigVMVariableNode>(InNode))
+	{
+		InverseAction.AddAction(FRigVMAddVariableNodeAction(VariableNode));
+		URigVMPin* ValuePin = VariableNode->FindPin(TEXT("Value"));
+		InverseAction.AddAction(FRigVMSetPinDefaultValueAction(ValuePin, ValuePin->GetDefaultValue()));
+	}
+	else if (URigVMParameterNode* ParameterNode = Cast<URigVMParameterNode>(InNode))
+	{
+		InverseAction.AddAction(FRigVMAddParameterNodeAction(ParameterNode));
+		URigVMPin* ValuePin = ParameterNode->FindPin(TEXT("Value"));
+		InverseAction.AddAction(FRigVMSetPinDefaultValueAction(ValuePin, ValuePin->GetDefaultValue()));
+	}
+	else if (URigVMCommentNode* CommentNode = Cast<URigVMCommentNode>(InNode))
+	{
+		InverseAction.AddAction(FRigVMAddCommentNodeAction(CommentNode));
+	}
+	else if (URigVMRerouteNode* RerouteNode = Cast<URigVMRerouteNode>(InNode))
+	{
+		InverseAction.AddAction(FRigVMAddRerouteNodeAction(RerouteNode));
+	}
+	else
+	{
+		ensure(false);
+	}
+
+	InverseActionKey.Set<FRigVMInverseAction>(InverseAction);
+}
+
+bool FRigVMRemoveNodeAction::Undo(URigVMController* InController)
+{
+	FRigVMActionWrapper InverseWrapper(InverseActionKey);
+	if (!InverseWrapper.GetAction()->Undo(InController))
+	{
+		return false;
+	}
+	return FRigVMBaseAction::Undo(InController);
+}
+
+bool FRigVMRemoveNodeAction::Redo(URigVMController* InController)
+{
+	if (!FRigVMBaseAction::Redo(InController))
+	{
+		return false;
+	}
+	FRigVMActionWrapper InverseWrapper(InverseActionKey);
+	return InverseWrapper.GetAction()->Redo(InController);
+}
+
+FRigVMSelectNodeAction::FRigVMSelectNodeAction()
+	: NodePath()
+	, bWasSelected(false)
+{
+
+}
+FRigVMSelectNodeAction::FRigVMSelectNodeAction(URigVMNode* InNode)
+	: NodePath(InNode->GetNodePath())
+	, bWasSelected(InNode->IsSelected())
+{
+}
+
+bool FRigVMSelectNodeAction::Undo(URigVMController* InController)
+{
+	if(!FRigVMBaseAction::Undo(InController))
+	{
+		return false;
+	}
+	return InController->SelectNodeByName(*NodePath, bWasSelected, false);
+}
+
+bool FRigVMSelectNodeAction::Redo(URigVMController* InController)
+{
+	if(!InController->SelectNodeByName(*NodePath, !bWasSelected, false))
+	{
+		return false;
+	}
+	return FRigVMBaseAction::Redo(InController);
+}
+
+FRigVMSetNodePositionAction::FRigVMSetNodePositionAction(URigVMNode* InNode, const FVector2D& InNewPosition)
+: NodePath(InNode->GetNodePath())
+, OldPosition(InNode->GetPosition())
+, NewPosition(InNewPosition)
+{
+}
+
+bool FRigVMSetNodePositionAction::Merge(const FRigVMBaseAction* Other)
+{
+	if (!FRigVMBaseAction::Merge(Other))
+	{
+		return false;
+	}
+
+	const FRigVMSetNodePositionAction* Action = (const FRigVMSetNodePositionAction*)Other;
+	if (NodePath != Action->NodePath)
+	{
+		return false;
+	}
+
+	NewPosition = Action->NewPosition;
+	return true;
+}
+
+bool FRigVMSetNodePositionAction::Undo(URigVMController* InController)
+{
+	if(!FRigVMBaseAction::Undo(InController))
+	{
+		return false;
+	}
+	return InController->SetNodePositionByName(*NodePath, OldPosition, false);
+}
+
+bool FRigVMSetNodePositionAction::Redo(URigVMController* InController)
+{
+	if(!InController->SetNodePositionByName(*NodePath, NewPosition, false))
+	{
+		return false;
+	}
+	return FRigVMBaseAction::Redo(InController);
+}
+
+FRigVMSetNodeSizeAction::FRigVMSetNodeSizeAction(URigVMNode* InNode, const FVector2D& InNewSize)
+: NodePath(InNode->GetNodePath())
+, OldSize(InNode->GetSize())
+, NewSize(InNewSize)
+{
+}
+
+bool FRigVMSetNodeSizeAction::Merge(const FRigVMBaseAction* Other)
+{
+	if (!FRigVMBaseAction::Merge(Other))
+	{
+		return false;
+	}
+
+	const FRigVMSetNodeSizeAction* Action = (const FRigVMSetNodeSizeAction*)Other;
+	if (NodePath != Action->NodePath)
+	{
+		return false;
+	}
+
+	NewSize = Action->NewSize;
+	return true;
+}
+
+bool FRigVMSetNodeSizeAction::Undo(URigVMController* InController)
+{
+	if(!FRigVMBaseAction::Undo(InController))
+	{
+		return false;
+	}
+	return InController->SetNodeSizeByName(*NodePath, OldSize, false);
+}
+
+bool FRigVMSetNodeSizeAction::Redo(URigVMController* InController)
+{
+	if(!InController->SetNodeSizeByName(*NodePath, NewSize, false))
+	{
+		return false;
+	}
+	return FRigVMBaseAction::Redo(InController);
+}
+
+FRigVMSetNodeColorAction::FRigVMSetNodeColorAction(URigVMNode* InNode, const FLinearColor& InNewColor)
+: NodePath(InNode->GetNodePath())
+, OldColor(InNode->GetNodeColor())
+, NewColor(InNewColor)
+{
+}
+
+bool FRigVMSetNodeColorAction::Merge(const FRigVMBaseAction* Other)
+{
+	if (!FRigVMBaseAction::Merge(Other))
+	{
+		return false;
+	}
+
+	const FRigVMSetNodeColorAction* Action = (const FRigVMSetNodeColorAction*)Other;
+	if (NodePath != Action->NodePath)
+	{
+		return false;
+	}
+
+	NewColor = Action->NewColor;
+	return true;
+}
+
+bool FRigVMSetNodeColorAction::Undo(URigVMController* InController)
+{
+	if(!FRigVMBaseAction::Undo(InController))
+	{
+		return false;
+	}
+	return InController->SetNodeColorByName(*NodePath, OldColor, false);
+}
+
+bool FRigVMSetNodeColorAction::Redo(URigVMController* InController)
+{
+	if(!InController->SetNodeColorByName(*NodePath, NewColor, false))
+	{
+		return false;
+	}
+	return FRigVMBaseAction::Redo(InController);
+}
+
+FRigVMSetCommentTextAction::FRigVMSetCommentTextAction(URigVMCommentNode* InNode, const FString& InNewText)
+: NodePath(InNode->GetNodePath())
+, OldText(InNode->GetCommentText())
+, NewText(InNewText)
+{
+}
+
+bool FRigVMSetCommentTextAction::Undo(URigVMController* InController)
+{
+	if(!FRigVMBaseAction::Undo(InController))
+	{
+		return false;
+	}
+	return InController->SetCommentTextByName(*NodePath, OldText, false);
+}
+
+bool FRigVMSetCommentTextAction::Redo(URigVMController* InController)
+{
+	if(!InController->SetCommentTextByName(*NodePath, NewText, false))
+	{
+		return false;
+	}
+	return FRigVMBaseAction::Redo(InController);
+}
+
+FRigVMSetRerouteCompactnessAction::FRigVMSetRerouteCompactnessAction()
+	: NodePath()
+	, OldShowAsFullNode(false)
+	, NewShowAsFullNode(false)
+{
+
+}
+FRigVMSetRerouteCompactnessAction::FRigVMSetRerouteCompactnessAction(URigVMRerouteNode* InNode, bool InShowAsFullNode)
+	: NodePath(InNode->GetNodePath())
+	, OldShowAsFullNode(InNode->GetShowsAsFullNode())
+	, NewShowAsFullNode(InShowAsFullNode)
+{
+}
+
+bool FRigVMSetRerouteCompactnessAction::Undo(URigVMController* InController)
+{
+	if(!FRigVMBaseAction::Undo(InController))
+	{
+		return false;
+	}
+	return InController->SetRerouteCompactnessByName(*NodePath, OldShowAsFullNode, false);
+}
+
+bool FRigVMSetRerouteCompactnessAction::Redo(URigVMController* InController)
+{
+	if(!InController->SetRerouteCompactnessByName(*NodePath, NewShowAsFullNode, false))
+	{
+		return false;
+	}
+	return FRigVMBaseAction::Redo(InController);
+}
+
+FRigVMRenameVariableAction::FRigVMRenameVariableAction(const FName& InOldVariableName, const FName& InNewVariableName)
+: OldVariableName(InOldVariableName.ToString())
+, NewVariableName(InNewVariableName.ToString())
+{
+}
+
+bool FRigVMRenameVariableAction::Undo(URigVMController* InController)
+{
+	if(!FRigVMBaseAction::Undo(InController))
+	{
+		return false;
+	}
+	return InController->RenameVariable(*NewVariableName, *OldVariableName, false);
+}
+
+bool FRigVMRenameVariableAction::Redo(URigVMController* InController)
+{
+	if(!InController->RenameVariable(*OldVariableName, *NewVariableName, false))
+	{
+		return false;
+	}
+	return FRigVMBaseAction::Redo(InController);
+}
+
+FRigVMRenameParameterAction::FRigVMRenameParameterAction(const FName& InOldParameterName, const FName& InNewParameterName)
+: OldParameterName(InOldParameterName.ToString())
+, NewParameterName(InNewParameterName.ToString())
+{
+}
+
+bool FRigVMRenameParameterAction::Undo(URigVMController* InController)
+{
+	if(!FRigVMBaseAction::Undo(InController))
+	{
+		return false;
+	}
+	return InController->RenameParameter(*NewParameterName, *OldParameterName, false);
+}
+
+bool FRigVMRenameParameterAction::Redo(URigVMController* InController)
+{
+	if(!InController->RenameParameter(*OldParameterName, *NewParameterName, false))
+	{
+		return false;
+	}
+	return FRigVMBaseAction::Redo(InController);
+}
+
+FRigVMSetPinDefaultValueAction::FRigVMSetPinDefaultValueAction(URigVMPin* InPin, const FString& InNewDefaultValue)
+: PinPath(InPin->GetPinPath())
+, OldDefaultValue(InPin->GetDefaultValue())
+, NewDefaultValue(InNewDefaultValue)
+{
+}
+
+bool FRigVMSetPinDefaultValueAction::Merge(const FRigVMBaseAction* Other)
+{
+	if (!FRigVMBaseAction::Merge(Other))
+	{
+		return false;
+	}
+
+	const FRigVMSetPinDefaultValueAction* Action = (const FRigVMSetPinDefaultValueAction*)Other;
+	if (PinPath != Action->PinPath)
+	{
+		return false;
+	}
+
+	NewDefaultValue = Action->NewDefaultValue;
+	return true;
+}
+
+bool FRigVMSetPinDefaultValueAction::Undo(URigVMController* InController)
+{
+	if(!FRigVMBaseAction::Undo(InController))
+	{
+		return false;
+	}
+	return InController->SetPinDefaultValue(PinPath, OldDefaultValue, true, false);
+}
+
+bool FRigVMSetPinDefaultValueAction::Redo(URigVMController* InController)
+{
+	if(!InController->SetPinDefaultValue(PinPath, NewDefaultValue, true, false))
+	{
+		return false;
+	}
+	return FRigVMBaseAction::Redo(InController);
+}
+
+FRigVMInsertArrayPinAction::FRigVMInsertArrayPinAction(URigVMPin* InArrayPin, int32 InIndex, const FString& InNewDefaultValue)
+: ArrayPinPath(InArrayPin->GetPinPath())
+, Index(InIndex)
+, NewDefaultValue(InNewDefaultValue)
+{
+}
+
+bool FRigVMInsertArrayPinAction::Undo(URigVMController* InController)
+{
+	if(!FRigVMBaseAction::Undo(InController))
+	{
+		return false;
+	}
+	return InController->RemoveArrayPin(FString::Printf(TEXT("%s.%d"), *ArrayPinPath, Index), false);
+}
+
+bool FRigVMInsertArrayPinAction::Redo(URigVMController* InController)
+{
+	if(InController->InsertArrayPin(ArrayPinPath, Index, NewDefaultValue, false).IsEmpty())
+	{
+		return false;
+	}
+	return FRigVMBaseAction::Redo(InController);
+}
+
+FRigVMRemoveArrayPinAction::FRigVMRemoveArrayPinAction(URigVMPin* InArrayElementPin)
+: ArrayPinPath(InArrayElementPin->GetParentPin()->GetPinPath())
+, Index(InArrayElementPin->GetPinIndex())
+, DefaultValue(InArrayElementPin->GetDefaultValue())
+{
+}
+
+bool FRigVMRemoveArrayPinAction::Undo(URigVMController* InController)
+{
+	if (InController->InsertArrayPin(*ArrayPinPath, Index, DefaultValue, false).IsEmpty())
+	{
+		return false;
+	}
+	return FRigVMBaseAction::Undo(InController);
+}
+
+bool FRigVMRemoveArrayPinAction::Redo(URigVMController* InController)
+{
+	if (!FRigVMBaseAction::Redo(InController))
+	{
+		return false;
+	}
+	return InController->RemoveArrayPin(FString::Printf(TEXT("%s.%d"), *ArrayPinPath, Index), false);
+}
+
+FRigVMAddLinkAction::FRigVMAddLinkAction(URigVMPin* InOutputPin, URigVMPin* InInputPin)
+: OutputPinPath(InOutputPin->GetPinPath())
+, InputPinPath(InInputPin->GetPinPath())
+{
+}
+
+bool FRigVMAddLinkAction::Undo(URigVMController* InController)
+{
+	if(!FRigVMBaseAction::Undo(InController))
+	{
+		return false;
+	}
+	return InController->BreakLink(OutputPinPath, InputPinPath, false);
+}
+
+bool FRigVMAddLinkAction::Redo(URigVMController* InController)
+{
+	if(!InController->AddLink(OutputPinPath, InputPinPath, false))
+	{
+		return false;
+	}
+	return FRigVMBaseAction::Redo(InController);
+}
+
+FRigVMBreakLinkAction::FRigVMBreakLinkAction(URigVMPin* InOutputPin, URigVMPin* InInputPin)
+: OutputPinPath(InOutputPin->GetPinPath())
+, InputPinPath(InInputPin->GetPinPath())
+{
+}
+
+bool FRigVMBreakLinkAction::Undo(URigVMController* InController)
+{
+	if(!FRigVMBaseAction::Undo(InController))
+	{
+		return false;
+	}
+	return InController->AddLink(OutputPinPath, InputPinPath, false);
+}
+
+bool FRigVMBreakLinkAction::Redo(URigVMController* InController)
+{
+	if(!InController->BreakLink(OutputPinPath, InputPinPath, false))
+	{
+		return false;
+	}
+	return FRigVMBaseAction::Redo(InController);
+}
\ No newline at end of file
