Copying //UE4/Dev-Anim to Dev-Main (//UE4/Dev-Main) @ 8782600

diff --git a/Engine/Source/Editor/RigVMEditor/Private/RigVMModel/RigVMGraph.cpp b/Engine/Source/Editor/RigVMEditor/Private/RigVMModel/RigVMGraph.cpp
new file mode 100644
--- /dev/null
+++ b/Engine/Source/Editor/RigVMEditor/Private/RigVMModel/RigVMGraph.cpp
@@ -0,0 +1,196 @@
+// Copyright 1998-2019 Epic Games, Inc. All Rights Reserved.
+
+#include "RigVMModel/RigVMGraph.h"
+#include "RigVMModel/RigVMController.h"
+#include "UObject/Package.h"
+
+URigVMGraph::URigVMGraph()
+: Utils(this)
+{
+}
+
+const TArray<URigVMNode*>& URigVMGraph::GetNodes() const
+{
+	return Nodes;
+}
+
+const TArray<URigVMLink*>& URigVMGraph::GetLinks() const
+{
+	return Links;
+}
+
+TArray<FRigVMGraphVariableDescription> URigVMGraph::GetVariableDescriptions() const
+{
+	TArray<FRigVMGraphVariableDescription> Variables;
+	for (URigVMNode* Node : Nodes)
+	{
+		if (URigVMVariableNode* VariableNode = Cast<URigVMVariableNode>(Node))
+		{
+			Variables.AddUnique(VariableNode->GetVariableDescription());
+		}
+	}
+	return Variables;
+}
+
+TArray<FRigVMGraphParameterDescription> URigVMGraph::GetParameterDescriptions() const
+{
+	TArray<FRigVMGraphParameterDescription> Parameters;
+	for (URigVMNode* Node : Nodes)
+	{
+		if (URigVMParameterNode* ParameterNode = Cast<URigVMParameterNode>(Node))
+		{
+			Parameters.AddUnique(ParameterNode->GetParameterDescription());
+		}
+	}
+	return Parameters;
+}
+
+URigVMNode* URigVMGraph::FindNodeByName(const FName& InNodeName) const
+{
+	for (URigVMNode* Node : Nodes)
+	{
+		if (Node->GetFName() == InNodeName)
+		{
+			return Node;
+		}
+	}
+	return nullptr;
+}
+
+URigVMNode* URigVMGraph::FindNode(const FString& InNodePath) const
+{
+	// for now we don't support encapsulation
+	// so we don't walk recursively
+	return FindNodeByName(FName(*InNodePath));
+}
+
+URigVMPin* URigVMGraph::FindPin(const FString& InPinPath) const
+{
+	FString Left, Right;
+	if (!URigVMPin::SplitPinPathAtStart(InPinPath, Left, Right))
+	{
+		Left = InPinPath;
+	}
+
+	URigVMNode* Node = FindNode(Left);
+	if (Node)
+	{
+		return Node->FindPin(Right);
+	}
+
+	return nullptr;
+}
+
+URigVMLink* URigVMGraph::FindLink(const FString& InLinkPinPathRepresentation) const
+{
+	for(URigVMLink* Link : Links)
+	{
+		if(Link->GetPinPathRepresentation() == InLinkPinPathRepresentation)
+		{
+			return Link;
+		}
+	}
+	return nullptr;
+}
+
+bool URigVMGraph::IsNodeSelected(const FName& InNodeName) const
+{
+	return SelectedNodes.Contains(InNodeName);
+}
+
+const TArray<FName>& URigVMGraph::GetSelectNodes() const
+{
+	return SelectedNodes;
+}
+
+FRigVMGraphModifiedEvent& URigVMGraph::OnModified()
+{
+	return ModifiedEvent;
+}
+
+void URigVMGraph::Notify(ERigVMGraphNotifType InNotifType, URigVMGraph* InGraph, UObject* InSubject)
+{
+	ModifiedEvent.Broadcast(InNotifType, InGraph, InSubject);
+}
+
+bool URigVMGraph::IsNameAvailable(const FString& InName)
+{
+	for (URigVMNode* Node : Nodes)
+	{
+		if (Node->GetName() == InName)
+		{
+			return false;
+		}
+	}
+	return true;
+}
+
+void URigVMGraph::PrepareCycleChecking(URigVMPin* InPin, bool bAsInput)
+{
+	Utils.PrepareCycleChecking(InPin, bAsInput);
+}
+
+bool URigVMGraph::CanLink(URigVMPin* InSourcePin, URigVMPin* InTargetPin, FString* OutFailureReason)
+{
+	if (!URigVMPin::CanLink(InSourcePin, InTargetPin, OutFailureReason))
+	{
+		return false;
+	}
+
+	if(!((Utils.GetCycleCheckPin() == InSourcePin && Utils.GetCycleCheckPinIsInput() == false) ||
+		(Utils.GetCycleCheckPin() == InTargetPin && Utils.GetCycleCheckPinIsInput() == true)))
+	{
+		PrepareCycleChecking(InTargetPin, true);
+	}
+
+	if (Utils.GetCycleCheckPinIsInput())
+	{
+		if (Utils.IsNodeOnCycle(InSourcePin->GetNode()))
+		{
+			if (OutFailureReason)
+			{
+				*OutFailureReason = TEXT("Cannot create a cycle.");
+			}
+			return false;
+		}
+	}
+	else
+	{
+		if (Utils.IsNodeOnCycle(InTargetPin->GetNode()))
+		{
+			if (OutFailureReason)
+			{
+				*OutFailureReason = TEXT("Cannot create a cycle.");
+			}
+			return false;
+		}
+	}
+
+	return true;
+}
+
+#if WITH_EDITORONLY_DATA
+
+void URigVMGraph::PostLoad()
+{
+	Super::PostLoad();
+
+	URigVMController* Controller = NewObject<URigVMController>(GetTransientPackage());
+	Controller->SetGraph(this);
+
+#if WITH_EDITOR
+
+	Controller->DetachLinksFromPinObjects();
+
+	for (URigVMNode* Node : Nodes)
+	{
+		Controller->RepopulatePinsOnNode(Node);
+	}
+
+	Controller->ReattachLinksToPinObjects();
+
+#endif
+
+}
+
+#endif
diff --git a/Engine/Source/Editor/RigVMEditor/Private/RigVMModel/RigVMGraphUtils.cpp b/Engine/Source/Editor/RigVMEditor/Private/RigVMModel/RigVMGraphUtils.cpp
new file mode 100644
--- /dev/null
+++ b/Engine/Source/Editor/RigVMEditor/Private/RigVMModel/RigVMGraphUtils.cpp
@@ -0,0 +1,327 @@
+// Copyright 1998-2019 Epic Games, Inc. All Rights Reserved.
+
+#include "RigVMModel/RigVMGraphUtils.h"
+#include "RigVMModel/RigVMGraph.h"
+
+FRigVMGraphUtils::FRigVMGraphUtils()
+	: Graph(nullptr)
+	, LastCycleCheckingPin(nullptr)
+	, LastCycleCheckingPinIsInput(true)
+{
+}
+
+FRigVMGraphUtils::FRigVMGraphUtils(URigVMGraph* InGraph, URigVMPin* InCycleCheckingPin, bool InCycleCheckingPinIsInput)
+	: Graph(InGraph)
+	, LastCycleCheckingPin(InCycleCheckingPin)
+	, LastCycleCheckingPinIsInput(InCycleCheckingPinIsInput)
+{
+	ModifiedHandle = Graph->OnModified().AddLambda([this](ERigVMGraphNotifType NotifType, URigVMGraph*, UObject*) {
+
+		switch (NotifType)
+		{
+			case ERigVMGraphNotifType::GraphChanged:
+			case ERigVMGraphNotifType::NodeAdded:
+			case ERigVMGraphNotifType::NodeRemoved:
+			case ERigVMGraphNotifType::LinkAdded:
+			case ERigVMGraphNotifType::LinkRemoved:
+			case ERigVMGraphNotifType::PinDirectionChanged:
+			case ERigVMGraphNotifType::PinTypeChanged:
+			{
+				Reset();
+				break;
+			}
+			default:
+			{
+				break;
+			}
+		}
+
+	});
+}
+
+FRigVMGraphUtils::~FRigVMGraphUtils()
+{
+	if (Graph && ModifiedHandle.IsValid())
+	{
+		Graph->OnModified().Remove(ModifiedHandle);
+	}
+}
+
+void FRigVMGraphUtils::Reset()
+{
+	CycleWhiteList.Empty();
+	CycleGrayList.Empty();
+	CycleBlackList.Empty();
+	CycleDepthTraversal.Empty();
+	Cycle.Empty();
+
+	const TArray<URigVMNode*> Nodes = Graph->GetNodes();
+	CycleWhiteList.Reserve(Nodes.Num());
+	CycleGrayList.Reserve(Nodes.Num());
+	CycleBlackList.Reserve(Nodes.Num());
+
+	LastCycleCheckingPin = nullptr;
+	NodeIsOnCycle.Reset();
+	NodeIsOnCycle.Reserve(Nodes.Num());
+}
+
+bool FRigVMGraphUtils::TopologicalSort(TArray<URigVMNode*>& OutOrder, TArray<URigVMNode*>& OutPotentialCycle)
+{
+	OutPotentialCycle = FindCycle();
+	if (OutPotentialCycle.Num() > 0)
+	{
+		return false;
+	}
+
+	struct Local
+	{
+		static void VisitNode(URigVMNode* Node, TArray<bool>& Visited, TArray<URigVMNode*>& SortedNodes)
+		{
+			int32 Index = Node->GetNodeIndex();
+			if (Visited[Index])
+			{
+				return;
+			}
+
+			Visited[Index] = true;
+
+			TArray<URigVMNode*> Inputs = Node->GetLinkedSourceNodes();
+			for (URigVMNode* Input : Inputs)
+			{
+				VisitNode(Input, Visited, SortedNodes);
+			}
+
+			SortedNodes.Push(Node);
+
+			if (Node->IsMutable())
+			{
+				TArray<URigVMNode*> Outputs = Node->GetLinkedTargetNodes();
+				for (URigVMNode* Output : Outputs)
+				{
+					VisitNode(Output, Visited, SortedNodes);
+				}
+			}
+		}
+	};
+
+	const TArray<URigVMNode*> Nodes = Graph->GetNodes();
+
+	// find all of the left mutable nodes on the left
+	TArray<URigVMNode*> LeafNodes;
+	TArray<URigVMNode*> OutputParameterNodes;
+	for (URigVMNode* Node : Nodes)
+	{
+		if (Node->IsMutable())
+		{
+			bool bHasInputPins = false;
+			for (URigVMPin* Pin : Node->GetPins())
+			{
+				if (Pin->GetDirection() == ERigVMPinDirection::Input || Pin->GetDirection() == ERigVMPinDirection::IO)
+				{
+					bHasInputPins = true;
+					break;
+				}
+			}
+			if (!bHasInputPins)
+			{
+				LeafNodes.Add(Node);
+			}
+		}
+		else if (Node->ContributesToResult())
+		{
+			OutputParameterNodes.Add(Node);
+		}
+	}
+
+	LeafNodes.Append(OutputParameterNodes);
+
+	TArray<bool> NodeVisited;
+	NodeVisited.SetNumZeroed(Nodes.Num());
+	OutOrder.Reset();
+	for (URigVMNode* LeafNode : LeafNodes)
+	{
+		Local::VisitNode(LeafNode, NodeVisited, OutOrder);
+	}
+	return true;
+}
+
+int32 FRigVMGraphUtils::GetMaxDistanceToLeafOutput(URigVMNode* InNode) const
+{
+	check(Graph)
+
+	TArray<URigVMNode*> Outputs = InNode->GetLinkedTargetNodes();
+	if(Outputs.Num() == 0)
+	{
+		return 0;
+	}
+
+	int32 MaxDistance = 0;
+	for (URigVMNode* TargetNode : Outputs)
+	{
+		int32 Distance = GetMaxDistanceToLeafOutput(TargetNode);
+		MaxDistance = FMath::Max<int32>(Distance, MaxDistance);
+	}
+	return MaxDistance + 1;
+}
+
+// Finds the cycles as node index arrays
+TArray<URigVMNode*> FRigVMGraphUtils::FindCycle()
+{
+	check(Graph)
+
+	const TArray<URigVMNode*> Nodes = Graph->GetNodes();
+
+	// mark all nodes to be on the white list (the non-visited list)
+	for (int32 NodeIndex = 0; NodeIndex < Nodes.Num(); NodeIndex++)
+	{
+		CycleWhiteList.Add(NodeIndex);
+	}
+
+	while (CycleBlackList.Num() < Nodes.Num())
+	{
+		// start a new pass by finding the first node in the white list
+		int32 Node = -1;
+		for (const int32& WhiteListNode : CycleWhiteList)
+		{
+			Node = WhiteListNode;
+			CycleDepthTraversal.Add(Node, -1);
+			CycleWhiteList.Remove(Node);
+			CycleGrayList.Add(Node);
+			break;
+		}
+
+		if(IsNodeCyclic(Node))
+		{
+			break;
+		}
+	}
+
+	return Cycle;
+}
+
+void FRigVMGraphUtils::PrepareCycleChecking(URigVMPin* InCycleCheckingPin, bool InCycleCheckingPinIsInput)
+{
+	check(Graph)
+		
+	if (LastCycleCheckingPin == InCycleCheckingPin && LastCycleCheckingPinIsInput == InCycleCheckingPinIsInput)
+	{
+		return;
+	}
+
+	LastCycleCheckingPin = InCycleCheckingPin;
+	LastCycleCheckingPinIsInput = InCycleCheckingPinIsInput;
+
+	check(LastCycleCheckingPin)
+
+	const TArray<URigVMNode*> Nodes = Graph->GetNodes();
+	NodeIsOnCycle.SetNumZeroed(Nodes.Num());
+
+	struct Local
+	{
+		static void VisitNode(URigVMNode* Node, const TArray<URigVMNode*>& Nodes, TArray<bool>& Visited, bool bWalkInputs)
+		{
+			int32 Index = Node->GetNodeIndex();
+			if (Visited[Index])
+			{
+				return;
+			}
+
+			Visited[Index] = true;
+
+			TArray<URigVMNode*> LinkedNodes = bWalkInputs ? Node->GetLinkedSourceNodes() : Node->GetLinkedTargetNodes();
+			for (URigVMNode* LinkedNode : LinkedNodes)
+			{
+				VisitNode(LinkedNode, Nodes, Visited, bWalkInputs);
+			}
+		}
+	};
+
+	Local::VisitNode(LastCycleCheckingPin->GetNode(), Nodes, NodeIsOnCycle, !LastCycleCheckingPinIsInput);
+}
+
+bool FRigVMGraphUtils::IsNodeOnCycle(URigVMNode* InNode)
+{
+	check(LastCycleCheckingPin)
+	ensure(NodeIsOnCycle.Num() == Graph->GetNodes().Num());
+
+	struct Local
+	{
+		static bool TestNodeOnCycle(URigVMNode* Node, const TArray<URigVMNode*>& Nodes, TArray<bool>& Visited, bool bWalkInputs)
+		{
+			int32 Index = Node->GetNodeIndex();
+			if (Visited[Index])
+			{
+				return true;
+			}
+
+			TArray<URigVMNode*> LinkedNodes = bWalkInputs ? Node->GetLinkedSourceNodes() : Node->GetLinkedTargetNodes();
+			for (URigVMNode* LinkedNode : LinkedNodes)
+			{
+				if (TestNodeOnCycle(LinkedNode, Nodes, Visited, bWalkInputs))
+				{
+					Visited[Index] = true;
+					return true;
+				}
+			}
+
+			return false;
+		}
+	};
+
+	const TArray<URigVMNode*> Nodes = Graph->GetNodes();
+	return Local::TestNodeOnCycle(InNode, Nodes, NodeIsOnCycle, LastCycleCheckingPinIsInput);
+}
+
+// this performs a depth first traversal by walking the output
+// links of each node
+bool FRigVMGraphUtils::IsNodeCyclic(int32 NodeIndex)
+{
+	if (Graph == nullptr)
+	{
+		return false;
+	}
+
+	const TArray<URigVMNode*> Nodes = Graph->GetNodes();
+	URigVMNode* Node = Graph->GetNodes()[NodeIndex];
+
+	TArray<URigVMNode*> Outputs = Node->GetLinkedTargetNodes();
+	for (URigVMNode* Output : Outputs)
+	{
+		int32 NeighborIndex = Output->GetNodeIndex();
+		if (CycleBlackList.Contains(NeighborIndex))
+		{
+			continue;
+		}
+		else if (CycleWhiteList.Contains(NeighborIndex))
+		{
+			CycleDepthTraversal.Add(NeighborIndex, NodeIndex);
+			CycleWhiteList.Remove(NeighborIndex);
+			CycleGrayList.Add(NeighborIndex);
+			if(IsNodeCyclic(NeighborIndex))
+			{
+				return true;
+			}
+		}
+		else if (CycleGrayList.Contains(NeighborIndex))
+		{
+			// this means we've detected a cycle
+			while(NodeIndex != -1)
+			{
+				Cycle.Add(Nodes[NodeIndex]);
+				NodeIndex = CycleDepthTraversal.FindChecked(NodeIndex);
+				if (NodeIndex == NeighborIndex)
+				{
+					Cycle.Add(Nodes[NodeIndex]);
+					break;
+				}
+			}
+			return true;
+		}
+	}
+
+	// if the current node has no neighbors we clear the gray list
+	CycleGrayList.Remove(NodeIndex);
+	CycleBlackList.Add(NodeIndex);
+	return false;
+}
+
diff --git a/Engine/Source/Editor/RigVMEditor/Private/RigVMModel/RigVMLink.cpp b/Engine/Source/Editor/RigVMEditor/Private/RigVMModel/RigVMLink.cpp
new file mode 100644
--- /dev/null
+++ b/Engine/Source/Editor/RigVMEditor/Private/RigVMModel/RigVMLink.cpp
@@ -0,0 +1,82 @@
+// Copyright 1998-2019 Epic Games, Inc. All Rights Reserved.
+
+#include "RigVMModel/RigVMLink.h"
+#include "RigVMModel/RigVMGraph.h"
+
+URigVMGraph* URigVMLink::GetGraph() const
+{
+	return Cast<URigVMGraph>(GetOuter());
+}
+
+void URigVMLink::Serialize(FArchive& Ar)
+{
+	if (Ar.IsLoading())
+	{
+		Ar << SourcePinPath;
+		Ar << TargetPinPath;
+	}
+	else
+	{
+		FString SourcePinPathTemp;
+		if (SourcePin)
+		{
+			if (SourcePin->GetNode())
+			{
+				SourcePinPathTemp = SourcePin->GetPinPath();
+			}
+		}
+		FString TargetPinPathTemp;
+		if (TargetPin)
+		{
+			if (TargetPin->GetNode())
+			{
+				TargetPinPathTemp = TargetPin->GetPinPath();
+			}
+		}
+		Ar << SourcePinPathTemp;
+		Ar << TargetPinPathTemp;
+	}
+}
+
+int32 URigVMLink::GetLinkIndex() const
+{
+	int32 Index = INDEX_NONE;
+	URigVMGraph* Graph = GetGraph();
+	if (Graph != nullptr)
+	{
+		Graph->GetLinks().Find((URigVMLink*)this, Index);
+	}
+	return Index;
+
+}
+
+URigVMPin* URigVMLink::GetSourcePin()
+{
+	if (SourcePin == nullptr)
+	{
+		SourcePin = GetGraph()->FindPin(SourcePinPath);
+		if (SourcePin)
+		{
+			SourcePinPath.Empty();
+		}
+	}
+	return SourcePin;
+}
+
+URigVMPin* URigVMLink::GetTargetPin()
+{
+	if (TargetPin == nullptr)
+	{
+		TargetPin = GetGraph()->FindPin(TargetPinPath);
+		if (TargetPin)
+		{
+			TargetPinPath.Empty();
+		}
+	}
+	return TargetPin;
+}
+
+FString URigVMLink::GetPinPathRepresentation()
+{
+	return FString::Printf(TEXT("%s -> %s"), *GetSourcePin()->GetPinPath(), *GetTargetPin()->GetPinPath());
+}
diff --git a/Engine/Source/Editor/RigVMEditor/Private/RigVMModel/RigVMPin.cpp b/Engine/Source/Editor/RigVMEditor/Private/RigVMModel/RigVMPin.cpp
new file mode 100644
--- /dev/null
+++ b/Engine/Source/Editor/RigVMEditor/Private/RigVMModel/RigVMPin.cpp
@@ -0,0 +1,453 @@
+// Copyright 1998-2019 Epic Games, Inc. All Rights Reserved.
+
+#include "RigVMModel/RigVMPin.h"
+#include "RigVMModel/RigVMNode.h"
+#include "RigVMModel/RigVMGraph.h"
+#include "RigVMModel/RigVMLink.h"
+#include "UObject/Package.h"
+
+bool URigVMPin::SplitPinPathAtStart(const FString& InPinPath, FString& LeftMost, FString& Right)
+{
+	return InPinPath.Split(TEXT("."), &LeftMost, &Right, ESearchCase::IgnoreCase, ESearchDir::FromStart);
+}
+
+bool URigVMPin::SplitPinPathAtEnd(const FString& InPinPath, FString& Left, FString& RightMost)
+{
+	return InPinPath.Split(TEXT("."), &Left, &RightMost, ESearchCase::IgnoreCase, ESearchDir::FromEnd);
+}
+
+bool URigVMPin::SplitPinPath(const FString& InPinPath, TArray<FString>& Parts)
+{
+	int32 OriginalPartsCount = Parts.Num();
+	FString PinPathRemaining = InPinPath;
+	FString Left, Right;
+	while(SplitPinPathAtStart(PinPathRemaining, Left, Right))
+	{
+		Parts.Add(Left);
+		Left.Empty();
+		PinPathRemaining = Right;
+	}
+
+	if (!Right.IsEmpty())
+	{
+		Parts.Add(Right);
+	}
+
+	return Parts.Num() > OriginalPartsCount;
+}
+
+FString URigVMPin::JoinSplitPath(const FString& Left, const FString& Right)
+{
+	ensure(!Left.IsEmpty() && !Right.IsEmpty());
+	return Left + TEXT(".") + Right;
+}
+
+FString URigVMPin::JoinSplitPath(const TArray<FString>& InParts)
+{
+	if (InParts.Num() == 0)
+	{
+		return FString();
+	}
+
+	FString Result = InParts[0];
+	for (int32 PartIndex = 1; PartIndex < InParts.Num(); PartIndex++)
+	{
+		Result += TEXT(".") + InParts[PartIndex];
+	}
+
+	return Result;
+}
+
+URigVMPin::URigVMPin()
+	: Direction(ERigVMPinDirection::Invalid)
+	, bIsConstant(false)
+	, CPPType(FString())
+	, ScriptStruct(nullptr)
+	, ScriptStructPath(NAME_None)
+	, DefaultValue(FString())
+{
+}
+
+FString URigVMPin::GetPinPath() const
+{
+	URigVMPin* ParentPin = GetParentPin();
+	if(ParentPin)
+	{
+		return FString::Printf(TEXT("%s.%s"), *ParentPin->GetPinPath(), *GetName());
+	}
+
+	URigVMNode* Node = GetNode();
+	return FString::Printf(TEXT("%s.%s"), *Node->GetNodePath(), *GetName());
+}
+
+FString URigVMPin::GetSegmentPath() const
+{
+	URigVMPin* ParentPin = GetParentPin();
+	if (ParentPin)
+	{
+		FString ParentSegmentPath = ParentPin->GetSegmentPath();
+		if (ParentSegmentPath.IsEmpty())
+		{
+			return GetName();
+		}
+		return FString::Printf(TEXT("%s.%s"), *ParentSegmentPath, *GetName());
+	}
+	return FString();
+}
+
+ERigVMPinDirection URigVMPin::GetDirection() const
+{
+	return Direction;
+}
+
+bool URigVMPin::IsConstant() const
+{
+	return bIsConstant;
+}
+
+bool URigVMPin::IsStruct() const
+{
+	if (IsArray())
+	{
+		return false;
+	}
+	return GetScriptStruct() != nullptr;
+}
+
+bool URigVMPin::IsStructMember() const
+{
+	URigVMPin* ParentPin = GetParentPin();
+	if (ParentPin == nullptr)
+	{
+		return false;
+	}
+	return ParentPin->IsStruct();
+}
+
+bool URigVMPin::IsArray() const
+{
+	return CPPType.StartsWith(TEXT("TArray"));
+}
+
+bool URigVMPin::IsArrayElement() const
+{
+	URigVMPin* ParentPin = GetParentPin();
+	if (ParentPin == nullptr)
+	{
+		return false;
+	}
+	return ParentPin->IsArray();
+}
+
+int32 URigVMPin::GetPinIndex() const
+{
+	int32 Index = INDEX_NONE;
+
+	URigVMPin* ParentPin = GetParentPin();
+	if (ParentPin != nullptr)
+	{
+		ParentPin->GetSubPins().Find((URigVMPin*)this, Index);
+	}
+	else
+	{
+		URigVMNode* Node = GetNode();
+		if (Node != nullptr)
+		{
+			Node->GetPins().Find((URigVMPin*)this, Index);
+		}
+	}
+	return Index;
+}
+
+void URigVMPin::SetNameFromIndex()
+{
+	LowLevelRename(*FString::FormatAsNumber(GetPinIndex()));
+}
+
+int32 URigVMPin::GetArraySize() const
+{
+	return SubPins.Num();
+}
+
+FString URigVMPin::GetCPPType() const
+{
+	return CPPType;
+}
+
+FString URigVMPin::GetArrayElementCppType() const
+{
+	if (!IsArray())
+	{
+		return FString();
+	}
+	return CPPType.Mid(7, CPPType.Len() - 8);
+}
+
+bool URigVMPin::IsStringType() const
+{
+	return CPPType.Equals(TEXT("FString")) || CPPType.Equals(TEXT("FName"));
+}
+
+FString URigVMPin::GetDefaultValue() const
+{
+	if (IsArray())
+	{
+		if (SubPins.Num() > 0)
+		{
+			TArray<FString> ElementDefaultValues;
+			for (URigVMPin* SubPin : SubPins)
+			{
+				FString ElementDefaultValue = SubPin->GetDefaultValue();
+				if (SubPin->IsStringType())
+				{
+					ElementDefaultValue = TEXT("\"") + ElementDefaultValue + TEXT("\"");
+				}
+				ElementDefaultValues.Add(ElementDefaultValue);
+			}
+			return FString::Printf(TEXT("(%s)"), *FString::Join(ElementDefaultValues, TEXT(",")));
+		}
+	}
+	else if (IsStruct())
+	{
+		if (SubPins.Num() > 0)
+		{
+			TArray<FString> MemberDefaultValues;
+			for (URigVMPin* SubPin : SubPins)
+			{
+				FString MemberDefaultValue = SubPin->GetDefaultValue();
+				if (SubPin->IsStringType())
+				{
+					MemberDefaultValue = TEXT("\"") + MemberDefaultValue + TEXT("\"");
+				}
+				MemberDefaultValues.Add(FString::Printf(TEXT("%s=%s"), *SubPin->GetName(), *MemberDefaultValue));
+			}
+			return FString::Printf(TEXT("(%s)"), *FString::Join(MemberDefaultValues, TEXT(",")));
+		}
+	}
+
+	if (DefaultValue.IsEmpty() && (IsArray() || IsStruct()))
+	{
+		return TEXT("()");
+	}
+	
+	return DefaultValue;
+}
+
+FName URigVMPin::GetCustomWidgetName() const
+{
+	return CustomWidgetName;
+}
+
+UScriptStruct* URigVMPin::GetScriptStruct() const
+{
+	if (ScriptStruct == nullptr)
+	{
+		if (ScriptStructPath != NAME_None)
+		{
+			URigVMPin* MutableThis = (URigVMPin*)this;
+			MutableThis->ScriptStruct = FindObject<UScriptStruct>(ANY_PACKAGE, *ScriptStructPath.ToString());
+		}
+	}
+	return ScriptStruct;
+}
+
+URigVMPin* URigVMPin::GetParentPin() const
+{
+	return Cast<URigVMPin>(GetOuter());
+}
+
+URigVMPin* URigVMPin::GetRootPin() const
+{
+	URigVMPin* ParentPin = GetParentPin();
+	if (ParentPin == nullptr)
+	{
+		return (URigVMPin*)this;
+	}
+	return ParentPin->GetRootPin();
+}
+
+const TArray<URigVMPin*>& URigVMPin::GetSubPins() const
+{
+	return SubPins;
+}
+
+URigVMPin* URigVMPin::FindSubPin(const FString& InPinPath) const
+{
+	FString Left, Right;
+	if (!URigVMPin::SplitPinPathAtStart(InPinPath, Left, Right))
+	{
+		Left = InPinPath;
+	}
+
+	for (URigVMPin* Pin : SubPins)
+	{
+		if (Pin->GetName() == Left)
+		{
+			if (Right.IsEmpty())
+			{
+				return Pin;
+			}
+			return Pin->FindSubPin(Right);
+		}
+	}
+	return nullptr;
+}
+
+bool URigVMPin::IsLinkedTo(URigVMPin* InPin) const
+{
+	for (URigVMLink* Link : Links)
+	{
+		if (Link->GetSourcePin() == InPin || Link->GetTargetPin() == InPin)
+		{
+			return true;
+		}
+	}
+	return false;
+}
+
+const TArray<URigVMLink*>& URigVMPin::GetLinks() const
+{
+	return Links;
+}
+
+TArray<URigVMPin*> URigVMPin::GetLinkedSourcePins() const
+{
+	TArray<URigVMPin*> Pins;
+	for (URigVMLink* Link : Links)
+	{
+		if (Link->GetTargetPin() == this)
+		{
+			Pins.AddUnique(Link->GetSourcePin());
+		}
+	}
+	return Pins;
+}
+
+TArray<URigVMPin*> URigVMPin::GetLinkedTargetPins() const
+{
+	TArray<URigVMPin*> Pins;
+	for (URigVMLink* Link : Links)
+	{
+		if (Link->GetSourcePin() == this)
+		{
+			Pins.AddUnique(Link->GetTargetPin());
+		}
+	}
+	return Pins;
+}
+
+URigVMNode* URigVMPin::GetNode() const
+{
+	URigVMPin* ParentPin = GetParentPin();
+	if (ParentPin)
+	{
+		return ParentPin->GetNode();
+	}
+
+	URigVMNode* Node = Cast<URigVMNode>(GetOuter());
+	if(Node)
+	{
+		return Node;
+	}
+
+	return nullptr;
+}
+
+URigVMGraph* URigVMPin::GetGraph() const
+{
+	URigVMNode* Node = GetNode();
+	if(Node)
+	{
+		return Node->GetGraph();
+	}
+
+	return nullptr;
+}
+
+bool URigVMPin::CanLink(URigVMPin* InSourcePin, URigVMPin* InTargetPin, FString* OutFailureReason)
+{
+	if (InSourcePin == nullptr || InTargetPin == nullptr)
+	{
+		if (OutFailureReason)
+		{
+			*OutFailureReason = TEXT("One of the pins is nullptr.");
+		}
+		return false;
+	}
+
+	if (InSourcePin == InTargetPin)
+	{
+		if (OutFailureReason)
+		{
+			*OutFailureReason = TEXT("Source and target pins are the same.");
+		}
+		return false;
+	}
+
+	if (InSourcePin->GetNode() == InTargetPin->GetNode())
+	{
+		if (OutFailureReason)
+		{
+			*OutFailureReason = TEXT("Source and target pins are on the same node.");
+		}
+		return false;
+	}
+
+	if (InSourcePin->GetGraph() != InTargetPin->GetGraph())
+	{
+		if (OutFailureReason)
+		{
+			*OutFailureReason = TEXT("Source and target pins are in different graphs.");
+		}
+		return false;
+	}
+
+	if (InSourcePin->Direction != ERigVMPinDirection::Output &&
+		InSourcePin->Direction != ERigVMPinDirection::IO)
+	{
+		if (OutFailureReason)
+		{
+			*OutFailureReason = TEXT("Source pin is not an output.");
+		}
+		return false;
+	}
+
+	if (InTargetPin->Direction != ERigVMPinDirection::Input &&
+		InTargetPin->Direction != ERigVMPinDirection::IO)
+	{
+		if (OutFailureReason)
+		{
+			*OutFailureReason = TEXT("Target pin is not an input.");
+		}
+		return false;
+	}
+
+	if (InTargetPin->bIsConstant && !InSourcePin->bIsConstant)
+	{
+		if (OutFailureReason)
+		{
+			*OutFailureReason = TEXT("Cannot connect non-constants to constants.");
+		}
+		return false;
+	}
+
+	if (InSourcePin->CPPType != InTargetPin->CPPType)
+	{
+		if (OutFailureReason)
+		{
+			*OutFailureReason = TEXT("Source and target pin types are not compatible.");
+		}
+		return false;
+	}
+
+	if (InSourcePin->IsLinkedTo(InTargetPin))
+	{
+		if (OutFailureReason)
+		{
+			*OutFailureReason = TEXT("Source and target pins are already connected.");
+		}
+		return false;
+	}
+
+	return true;
+}
diff --git a/Engine/Source/Editor/RigVMEditor/Public/RigVMCompiler/RigVMCompiler.h b/Engine/Source/Editor/RigVMEditor/Public/RigVMCompiler/RigVMCompiler.h
new file mode 100644
--- /dev/null
+++ b/Engine/Source/Editor/RigVMEditor/Public/RigVMCompiler/RigVMCompiler.h
@@ -0,0 +1,79 @@
+// Copyright 1998-2019 Epic Games, Inc. All Rights Reserved.
+
+#pragma once
+
+#include "CoreMinimal.h"
+
+#include "RigVMModel/RigVMGraph.h"
+#include "RigVMModel/Nodes/RigVMStructNode.h"
+#include "RigVMModel/Nodes/RigVMVariableNode.h"
+#include "RigVMModel/Nodes/RigVMParameterNode.h"
+#include "RigVMModel/Nodes/RigVMCommentNode.h"
+#include "RigVMModel/Nodes/RigVMRerouteNode.h"
+#include "RigVMCore/RigVM.h"
+
+#include "RigVMCompiler.generated.h"
+
+USTRUCT(BlueprintType)
+struct RIGVMEDITOR_API FRigVMCompileSettings
+{
+	GENERATED_BODY()
+
+public:
+
+	FRigVMCompileSettings();
+
+	UPROPERTY(BlueprintReadWrite, Category = FRigVMCompileSettings)
+	bool SurpressInfoMessage;
+
+	UPROPERTY(BlueprintReadWrite, Category = FRigVMCompileSettings)
+	bool SurpressWarnings;
+
+	UPROPERTY(BlueprintReadWrite, Category = FRigVMCompileSettings)
+	bool SurpressErrors;
+};
+
+UCLASS(BlueprintType)
+class RIGVMEDITOR_API URigVMCompiler : public UObject
+{
+	GENERATED_BODY()
+
+public:
+
+	URigVMCompiler();
+
+	UPROPERTY(BlueprintReadWrite, Category = FRigVMCompiler)
+	FRigVMCompileSettings Settings;
+
+	UFUNCTION(BlueprintCallable, Category = FRigVMCompiler)
+	bool Compile(URigVMGraph* InGraph, URigVM* OutVM);
+
+private:
+
+	FString GetPinHash(URigVMPin* InPin);
+	UScriptStruct* GetScriptStructForCPPType(const FString& InCPPType);
+	uint16 GetElementSizeFromCPPType(const FString& InCPPType, UScriptStruct* InScriptStruct);
+	FRigVMOperand AddRegisterForPin(URigVMPin* InPin, URigVM* OutVM);
+
+	void ReportInfo(const FString& InMessage);
+	void ReportWarning(const FString& InMessage);
+	void ReportError(const FString& InMessage);
+
+	template <typename FmtType, typename... Types>
+	void ReportInfof(const FmtType& Fmt, Types... Args)
+	{
+		ReportInfo(FString::Printf(Fmt, Args...));
+	}
+
+	template <typename FmtType, typename... Types>
+	void ReportWarningf(const FmtType& Fmt, Types... Args)
+	{
+		ReportWarning(FString::Printf(Fmt, Args...));
+	}
+
+	template <typename FmtType, typename... Types>
+	void ReportErrorf(const FmtType& Fmt, Types... Args)
+	{
+		ReportError(FString::Printf(Fmt, Args...));
+	}
+};
diff --git a/Engine/Source/Editor/RigVMEditor/Public/RigVMModel/Nodes/RigVMRerouteNode.h b/Engine/Source/Editor/RigVMEditor/Public/RigVMModel/Nodes/RigVMRerouteNode.h
new file mode 100644
--- /dev/null
+++ b/Engine/Source/Editor/RigVMEditor/Public/RigVMModel/Nodes/RigVMRerouteNode.h
@@ -0,0 +1,43 @@
+// Copyright 1998-2019 Epic Games, Inc. All Rights Reserved.
+
+#pragma once
+
+#include "RigVMModel/RigVMNode.h"
+#include "RigVMRerouteNode.generated.h"
+
+/**
+ * A reroute node is used to visually improve the 
+ * data flow within a Graph. Reroutes are purely 
+ * cosmetic and have no impact on the resulting
+ * VM whatsoever. Reroutes can furthermore be
+ * displayed as full nodes or as small circles.
+ */
+UCLASS(BlueprintType)
+class RIGVMEDITOR_API URigVMRerouteNode : public URigVMNode
+{
+	GENERATED_BODY()
+
+public:
+
+	// Default constructor
+	URigVMRerouteNode();
+
+	// Override of node title
+	virtual FString GetNodeTitle() const override { return RerouteName; }
+
+	// Returns true if this node should be shown as a full node,
+	// or false if this node should be shown as a small circle.
+	UFUNCTION(BlueprintCallable, Category = RigVMRerouteNode)
+	bool GetShowsAsFullNode() const;
+
+private:
+
+	static const FString RerouteName;
+	static const FString ValueName;
+
+	UPROPERTY()
+	bool bShowAsFullNode;
+
+	friend class URigVMController;
+};
+
diff --git a/Engine/Source/Editor/RigVMEditor/Public/RigVMModel/Nodes/RigVMStructNode.h b/Engine/Source/Editor/RigVMEditor/Public/RigVMModel/Nodes/RigVMStructNode.h
new file mode 100644
--- /dev/null
+++ b/Engine/Source/Editor/RigVMEditor/Public/RigVMModel/Nodes/RigVMStructNode.h
@@ -0,0 +1,39 @@
+// Copyright 1998-2019 Epic Games, Inc. All Rights Reserved.
+
+#pragma once
+
+#include "RigVMModel/RigVMNode.h"
+#include "RigVMStructNode.generated.h"
+
+/**
+ * The Struct Node represents a Function Invocation of a RIGVM_METHOD
+ * declared on a USTRUCT. Struct Nodes have input / output pins for all
+ * struct UPROPERTY members.
+ */
+UCLASS(BlueprintType)
+class RIGVMEDITOR_API URigVMStructNode : public URigVMNode
+{
+	GENERATED_BODY()
+
+public:
+
+	// Returns the UStruct for this struct node
+	// (the struct declaring the RIGVM_METHOD)
+	UFUNCTION(BlueprintCallable, Category = RigVMStructNode)
+	UScriptStruct* GetScriptStruct() const;
+
+	// Returns the name of the declared RIGVM_METHOD
+	UFUNCTION(BlueprintCallable, Category = RigVMStructNode)
+	FName GetMethodName() const;
+
+private:
+
+	UPROPERTY()
+	UScriptStruct* ScriptStruct;
+
+	UPROPERTY()
+	FName MethodName;
+
+	friend class URigVMController;
+};
+
diff --git a/Engine/Source/Editor/RigVMEditor/Public/RigVMModel/Nodes/RigVMVariableNode.h b/Engine/Source/Editor/RigVMEditor/Public/RigVMModel/Nodes/RigVMVariableNode.h
new file mode 100644
--- /dev/null
+++ b/Engine/Source/Editor/RigVMEditor/Public/RigVMModel/Nodes/RigVMVariableNode.h
@@ -0,0 +1,98 @@
+// Copyright 1998-2019 Epic Games, Inc. All Rights Reserved.
+
+#pragma once
+
+#include "RigVMModel/RigVMNode.h"
+#include "RigVMVariableNode.generated.h"
+
+/**
+ * The variable description is used to convey information
+ * about unique variables within a Graph. Multiple Variable
+ * Nodes can share the same variable description.
+ */
+USTRUCT(BlueprintType)
+struct FRigVMGraphVariableDescription
+{
+	GENERATED_BODY()
+
+public:
+
+	// comparison operator
+	bool operator ==(const FRigVMGraphVariableDescription& Other) const
+	{
+		return Name == Other.Name;
+	}
+
+	// The name of the variable
+	UPROPERTY(BlueprintReadOnly, Category = RigVMGraphVariableDescription)
+	FName Name;
+
+	// The C++ data type of the variable
+	UPROPERTY(BlueprintReadOnly, Category = RigVMGraphVariableDescription)
+	FString CPPType;
+
+	// The Struct of the C++ data type of the variable (or nullptr)
+	UPROPERTY(BlueprintReadOnly, Category = RigVMGraphVariableDescription)
+	UScriptStruct* ScriptStruct;
+
+	// The default value of the variable
+	UPROPERTY(BlueprintReadOnly, Category = RigVMGraphVariableDescription)
+	FString DefaultValue;
+};
+
+/**
+ * The Variable Node represents a mutable value / local state within the
+ * the Function / Graph. Variable Node's can be a getter or a setter.
+ * Getters are pure nodes with just an output value pin, while setters
+ * are mutable nodes with an execute and input value pin.
+ */
+UCLASS(BlueprintType)
+class RIGVMEDITOR_API URigVMVariableNode : public URigVMNode
+{
+	GENERATED_BODY()
+
+public:
+
+	// Default constructor
+	URigVMVariableNode();
+
+	// Override of node title
+	virtual FString GetNodeTitle() const;
+
+	// Returns the name of the variable
+	UFUNCTION(BlueprintCallable, Category = RigVMVariableNode)
+	FName GetVariableName() const;
+
+	// Returns true if this node is a variable getter
+	UFUNCTION(BlueprintCallable, Category = RigVMVariableNode)
+	bool IsGetter() const;
+
+	// Returns the C++ data type of the variable
+	UFUNCTION(BlueprintCallable, Category = RigVMVariableNode)
+	FString GetCPPType() const;
+
+	// Returns the C++ data type struct of the variable (or nullptr)
+	UFUNCTION(BlueprintCallable, Category = RigVMVariableNode)
+	UScriptStruct* GetScriptStruct() const;
+
+	// Returns the default value of the variable as a string
+	UFUNCTION(BlueprintCallable, Category = RigVMVariableNode)
+	FString GetDefaultValue() const;
+
+	// Returns this variable node's variable description
+	UFUNCTION(BlueprintCallable, Category = RigVMVariableNode)
+	FRigVMGraphVariableDescription GetVariableDescription() const;
+
+	// Override of node title
+	virtual FLinearColor GetNodeColor() const override { return FLinearColor::Blue; }
+
+private:
+
+	static const FString ValueName;
+
+	UPROPERTY()
+	FName VariableName;
+
+	friend class URigVMController;
+};
+
diff --git a/Engine/Source/Editor/RigVMEditor/Public/RigVMModel/RigVMController.h b/Engine/Source/Editor/RigVMEditor/Public/RigVMModel/RigVMController.h
new file mode 100644
--- /dev/null
+++ b/Engine/Source/Editor/RigVMEditor/Public/RigVMModel/RigVMController.h
@@ -0,0 +1,358 @@
+// Copyright 1998-2019 Epic Games, Inc. All Rights Reserved.
+
+#pragma once
+
+#include "RigVMGraph.h"
+#include "RigVMModel/Nodes/RigVMStructNode.h"
+#include "RigVMModel/Nodes/RigVMVariableNode.h"
+#include "RigVMModel/Nodes/RigVMParameterNode.h"
+#include "RigVMModel/Nodes/RigVMCommentNode.h"
+#include "RigVMModel/Nodes/RigVMRerouteNode.h"
+#include "RigVMController.generated.h"
+
+class URigVMActionStack;
+
+/**
+ * The Controller is the sole authority to perform changes
+ * on the Graph. The Controller itself is stateless.
+ * The Controller offers a Modified event to subscribe to 
+ * for user interface views - so they can be informed about
+ * any change that's happening within the Graph.
+ * The Controller routes all changes through the Graph itself,
+ * so you can have N Controllers performing edits on 1 Graph,
+ * and N Views subscribing to 1 Controller.
+ * In Python you can also subscribe to this event to be 
+ * able to react to topological changes of the Graph there.
+ */
+UCLASS(BlueprintType)
+class RIGVMEDITOR_API URigVMController : public UObject
+{
+	GENERATED_UCLASS_BODY()
+
+public:
+
+	// Default constructor
+	URigVMController();
+
+	// Default destructor
+	~URigVMController();
+
+	// Returns the currently edited Graph of this controller.
+	UFUNCTION(BlueprintCallable, Category = RigVMController)
+	URigVMGraph* GetGraph() const;
+
+	// Sets the currently edited Graph of this controller.
+	// This causes a GraphChanged modified event.
+	UFUNCTION(BlueprintCallable, Category = RigVMController)
+	void SetGraph(URigVMGraph* InGraph);
+
+	// The Modified event used to subscribe to changes
+	// happening within the Graph. This is broadcasted to 
+	// for any change happening - not only the changes 
+	// performed by this Controller - so it can be used
+	// for UI Views to react accordingly.
+	FRigVMGraphModifiedEvent& OnModified();
+
+	// Enables or disables the error reporting of this Controller.
+	UFUNCTION(BlueprintCallable, Category = RigVMController)
+	void EnableReporting(bool bEnabled = true) { bReportWarningsAndErrors = bEnabled; }
+
+#if WITH_EDITOR
+	// Note: The functions below are scoped with WITH_EDITOR since we are considering
+	// to move this code into the runtime in the future. Right now there's a dependency
+	// on the metadata of the USTRUCT - which is only available in the editor.
+
+	// Adds a Function / Struct Node to the edited Graph.
+	// StructNode represent a RIGVM_METHOD declaration on a USTRUCT.
+	// This causes a NodeAdded modified event.
+	UFUNCTION(BlueprintCallable, Category = RigVMController)
+	URigVMStructNode* AddStructNode(UScriptStruct* InScriptStruct, const FName& InMethodName, const FVector2D& InPosition = FVector2D::ZeroVector, const FString& InNodeName = TEXT(""), bool bUndo = true);
+
+	// Adds a Function / Struct Node to the edited Graph given its struct object path name.
+	// StructNode represent a RIGVM_METHOD declaration on a USTRUCT.
+	// This causes a NodeAdded modified event.
+	UFUNCTION(BlueprintCallable, Category = RigVMController)
+	URigVMStructNode* AddStructNodeFromStructPath(const FString& InScriptStructPath, const FName& InMethodName, const FVector2D& InPosition = FVector2D::ZeroVector, const FString& InNodeName = TEXT(""), bool bUndo = true);
+
+	// Adds a Variable Node to the edited Graph.
+	// Variables represent local work state for the function and
+	// can be read from and written to. 
+	// This causes a NodeAdded modified event.
+	UFUNCTION(BlueprintCallable, Category = RigVMController)
+	URigVMVariableNode* AddVariableNode(const FName& InVariableName, const FString& InCPPType, UScriptStruct* InScriptStruct, bool bIsGetter, const FString& InDefaultValue, const FVector2D& InPosition = FVector2D::ZeroVector, const FString& InNodeName = TEXT(""), bool bUndo = true);
+
+	// Adds a Variable Node to the edited Graph given a struct object path name.
+	// Variables represent local work state for the function and
+	// can be read from (bIsGetter == true) or written to (bIsGetter == false). 
+	// This causes a NodeAdded modified event.
+	UFUNCTION(BlueprintCallable, Category = RigVMController)
+	URigVMVariableNode* AddVariableNodeFromStructPath(const FName& InVariableName, const FString& InCPPType, const FString& InScriptStructPath, bool bIsGetter, const FString& InDefaultValue, const FVector2D& InPosition = FVector2D::ZeroVector, const FString& InNodeName = TEXT(""), bool bUndo = true);
+
+	// Adds a Parameter Node to the edited Graph.
+	// Parameters represent input or output arguments to the Graph / Function.
+	// Input Parameters are constant values / literals.
+	// This causes a NodeAdded modified event.
+	UFUNCTION(BlueprintCallable, Category = RigVMController)
+	URigVMParameterNode* AddParameterNode(const FName& InParameterName, const FString& InCPPType, UScriptStruct* InScriptStruct, bool bIsInput, const FString& InDefaultValue, const FVector2D& InPosition = FVector2D::ZeroVector, const FString& InNodeName = TEXT(""), bool bUndo = true);
+
+	// Adds a Parameter Node to the edited Graph given a struct object path name.
+	// Parameters represent input or output arguments to the Graph / Function.
+	// Input Parameters are constant values / literals.
+	// This causes a NodeAdded modified event.
+	UFUNCTION(BlueprintCallable, Category = RigVMController)
+	URigVMParameterNode* AddParameterNodeFromStructPath(const FName& InParameterName, const FString& InCPPType, const FString& InScriptStructPath, bool bIsInput, const FString& InDefaultValue, const FVector2D& InPosition = FVector2D::ZeroVector, const FString& InNodeName = TEXT(""), bool bUndo = true);
+
+	// Adds a Comment Node to the edited Graph.
+	// Comments can be used to annotate the Graph.
+	// This causes a NodeAdded modified event.
+	UFUNCTION(BlueprintCallable, Category = RigVMController)
+	URigVMCommentNode* AddCommentNode(const FString& InCommentText, const FVector2D& InPosition = FVector2D::ZeroVector, const FVector2D& InSize = FVector2D(400.f, 300.f), const FLinearColor& InColor = FLinearColor::Black, const FString& InNodeName = TEXT(""), bool bUndo = true);
+
+	// Adds a Reroute Node on an existing Link to the edited Graph.
+	// Reroute Nodes can be used to visually improve the data flow,
+	// they don't require any additional memory though and are purely
+	// cosmetic. This causes a NodeAdded modified event.
+	UFUNCTION(BlueprintCallable, Category = RigVMController)
+	URigVMRerouteNode* AddRerouteNodeOnLink(URigVMLink* InLink, bool bShowAsFullNode, const FVector2D& InPosition = FVector2D::ZeroVector, const FString& InNodeName = TEXT(""), bool bUndo = true);
+
+	// Adds a Reroute Node on an existing Link to the edited Graph given the Link's string representation.
+	// Reroute Nodes can be used to visually improve the data flow,
+	// they don't require any additional memory though and are purely
+	// cosmetic. This causes a NodeAdded modified event.
+	UFUNCTION(BlueprintCallable, Category = RigVMController)
+	URigVMRerouteNode* AddRerouteNodeOnLinkPath(const FString& InLinkPinPathRepresentation, bool bShowAsFullNode, const FVector2D& InPosition = FVector2D::ZeroVector, const FString& InNodeName = TEXT(""), bool bUndo = true);
+
+	// Adds a Reroute Node on an existing Pin to the editor Graph.
+	// Reroute Nodes can be used to visually improve the data flow,
+	// they don't require any additional memory though and are purely
+	// cosmetic. This causes a NodeAdded modified event.
+	UFUNCTION(BlueprintCallable, Category = RigVMController)
+	URigVMRerouteNode* AddRerouteNodeOnPin(const FString& InPinPath, bool bShowAsFullNode, bool bAsInput, const FVector2D& InPosition = FVector2D::ZeroVector, const FString& InNodeName = TEXT(""), bool bUndo = true);
+
+	// Un-does the last action on the stack.
+	// Note: This should really only be used for unit tests,
+	// use the GEditor's main Undo method instead.
+	UFUNCTION(BlueprintCallable, Category = RigVMController)
+	bool Undo();
+
+	// Re-does the last action on the stack.
+	// Note: This should really only be used for unit tests,
+	// use the GEditor's main Undo method instead.
+	UFUNCTION(BlueprintCallable, Category = RigVMController)
+	bool Redo();
+
+	// Opens an undo bracket / scoped transaction for
+	// a series of actions to be performed as one step on the 
+	// Undo stack. This is primarily useful for Python.
+	UFUNCTION(BlueprintCallable, Category = RigVMController)
+	bool OpenUndoBracket(const FString& InTitle);
+
+	// Closes an undo bracket / scoped transaction.
+	// This is primarily useful for Python.
+	UFUNCTION(BlueprintCallable, Category = RigVMController)
+	bool CloseUndoBracket();
+
+#endif
+
+	// Removes a node from the graph
+	// This causes a NodeRemoved modified event.
+	UFUNCTION(BlueprintCallable, Category = RigVMController)
+	bool RemoveNode(URigVMNode* InNode, bool bUndo = true);
+
+	// Removes a node from the graph given the node's name.
+	// This causes a NodeRemoved modified event.
+	UFUNCTION(BlueprintCallable, Category = RigVMController)
+	bool RemoveNodeByName(const FName& InNodeName, bool bUndo = true);
+
+	// Selects a single node in the graph.
+	// This causes a NodeSelected / NodeDeselected modified event.
+	UFUNCTION(BlueprintCallable, Category = RigVMController)
+	bool SelectNode(URigVMNode* InNode, bool bSelect = true, bool bUndo = true);
+
+	// Selects a single node in the graph by name.
+	// This causes a NodeSelected / NodeDeselected modified event.
+	UFUNCTION(BlueprintCallable, Category = RigVMController)
+	bool SelectNodeByName(const FName& InNodeName, bool bSelect = true, bool bUndo = true);
+
+	// Deselects all currently selected nodes in the graph.
+	// This might cause several NodeDeselected modified event.
+	UFUNCTION(BlueprintCallable, Category = RigVMController)
+	bool ClearNodeSelection(bool bUndo = true);
+
+	// Sets the position of a node in the graph.
+	// This causes a NodePositionChanged modified event.
+	UFUNCTION(BlueprintCallable, Category = RigVMController)
+	bool SetNodePosition(URigVMNode* InNode, const FVector2D& InPosition, bool bUndo = true, bool bMergeUndoAction = false);
+
+	// Sets the position of a node in the graph by name.
+	// This causes a NodePositionChanged modified event.
+	UFUNCTION(BlueprintCallable, Category = RigVMController)
+	bool SetNodePositionByName(const FName& InNodeName, const FVector2D& InPosition, bool bUndo = true, bool bMergeUndoAction = false);
+
+	// Sets the size of a node in the graph.
+	// This causes a NodeSizeChanged modified event.
+	UFUNCTION(BlueprintCallable, Category = RigVMController)
+	bool SetNodeSize(URigVMNode* InNode, const FVector2D& InSize, bool bUndo = true, bool bMergeUndoAction = false);
+
+	// Sets the size of a node in the graph by name.
+	// This causes a NodeSizeChanged modified event.
+	UFUNCTION(BlueprintCallable, Category = RigVMController)
+	bool SetNodeSizeByName(const FName& InNodeName, const FVector2D& InSize, bool bUndo = true, bool bMergeUndoAction = false);
+
+	// Sets the color of a node in the graph.
+	// This causes a NodeColorChanged modified event.
+	UFUNCTION(BlueprintCallable, Category = RigVMController)
+	bool SetNodeColor(URigVMNode* InNode, const FLinearColor& InColor, bool bUndo = true, bool bMergeUndoAction = false);
+
+	// Sets the color of a node in the graph by name.
+	// This causes a NodeColorChanged modified event.
+	UFUNCTION(BlueprintCallable, Category = RigVMController)
+	bool SetNodeColorByName(const FName& InNodeName, const FLinearColor& InColor, bool bUndo = true, bool bMergeUndoAction = false);
+
+	// Sets the comment text of a comment node in the graph.
+	// This causes a CommentTextChanged modified event.
+	UFUNCTION(BlueprintCallable, Category = RigVMController)
+	bool SetCommentText(URigVMNode* InNode, const FString& InCommentText, bool bUndo = true);
+
+	// Sets the comment text of a comment node in the graph by name.
+	// This causes a CommentTextChanged modified event.
+	UFUNCTION(BlueprintCallable, Category = RigVMController)
+	bool SetCommentTextByName(const FName& InNodeName, const FString& InCommentText, bool bUndo = true);
+
+	// Sets the compactness of a reroute node in the graph.
+	// This causes a RerouteCompactnessChanged modified event.
+	UFUNCTION(BlueprintCallable, Category = RigVMController)
+	bool SetRerouteCompactness(URigVMNode* InNode, bool bShowAsFullNode, bool bUndo = true);
+
+	// Sets the compactness of a reroute node in the graph by name.
+	// This causes a RerouteCompactnessChanged modified event.
+	UFUNCTION(BlueprintCallable, Category = RigVMController)
+	bool SetRerouteCompactnessByName(const FName& InNodeName, bool bShowAsFullNode, bool bUndo = true);
+
+	// Renames a variable in the graph.
+	// This causes a VariableRenamed modified event.
+	UFUNCTION(BlueprintCallable, Category = RigVMController)
+	bool RenameVariable(const FName& InOldName, const FName& InNewName, bool bUndo = true);
+
+	// Renames a parameter in the graph.
+	// This causes a ParameterRenamed modified event.
+	UFUNCTION(BlueprintCallable, Category = RigVMController)
+	bool RenameParameter(const FName& InOldName, const FName& InNewName, bool bUndo = true);
+
+	// Returns the default value of a pin given its pinpath.
+	UFUNCTION(BlueprintCallable, Category = RigVMController)
+	FString GetPinDefaultValue(const FString& InPinPath);
+
+	// Sets the default value of a pin given its pinpath.
+	// This causes a PinDefaultValueChanged modified event.
+	UFUNCTION(BlueprintCallable, Category = RigVMController)
+	bool SetPinDefaultValue(const FString& InPinPath, const FString& InDefaultValue, bool bResizeArrays = true, bool bUndo = true, bool bMergeUndoAction = false);
+
+	// Adds an array element pin to the end of an array pin.
+	// This causes a PinArraySizeChanged modified event.
+	UFUNCTION(BlueprintCallable, Category = RigVMController)
+	FString AddArrayPin(const FString& InArrayPinPath, const FString& InDefaultValue = TEXT(""), bool bUndo = true);
+
+	// Duplicates an array element pin.
+	// This causes a PinArraySizeChanged modified event.
+	UFUNCTION(BlueprintCallable, Category = RigVMController)
+	FString DuplicateArrayPin(const FString& InArrayElementPinPath, bool bUndo = true);
+
+	// Inserts an array element pin into an array pin.
+	// This causes a PinArraySizeChanged modified event.
+	UFUNCTION(BlueprintCallable, Category = RigVMController)
+	FString InsertArrayPin(const FString& InArrayPinPath, int32 InIndex = -1, const FString& InDefaultValue = TEXT(""), bool bUndo = true);
+
+	// Removes an array element pin from an array pin.
+	// This causes a PinArraySizeChanged modified event.
+	UFUNCTION(BlueprintCallable, Category = RigVMController)
+	bool RemoveArrayPin(const FString& InArrayElementPinPath, bool bUndo = true);
+
+	// Removes all (but one) array element pin from an array pin.
+	// This causes a PinArraySizeChanged modified event.
+	UFUNCTION(BlueprintCallable, Category = RigVMController)
+	bool ClearArrayPin(const FString& InArrayPinPath, bool bUndo = true);
+
+	// Adds a link to the graph.
+	// This causes a LinkAdded modified event.
+	UFUNCTION(BlueprintCallable, Category = RigVMController)
+	bool AddLink(const FString& InOutputPinPath, const FString& InInputPinPath, bool bUndo = true);
+
+	// Removes a link from the graph.
+	// This causes a LinkRemoved modified event.
+	UFUNCTION(BlueprintCallable, Category = RigVMController)
+	bool BreakLink(const FString& InOutputPinPath, const FString& InInputPinPath, bool bUndo = true);
+
+	// Removes all links on a given pin from the graph.
+	// This might cause multiple LinkRemoved modified event.
+	UFUNCTION(BlueprintCallable, Category = RigVMController)
+	bool BreakAllLinks(const FString& InPinPath, bool bAsInput = true, bool bUndo = true);
+
+private:
+
+	UPROPERTY(BlueprintReadOnly, Category = RigVMController, meta = (ScriptName = "ModifiedEvent", AllowPrivateAccess = "true"))
+	FRigVMGraphModifiedDynamicEvent ModifiedEventDynamic;
+
+	FRigVMGraphModifiedEvent ModifiedEventStatic;
+	void Notify(ERigVMGraphNotifType InNotifType, URigVMGraph* InGraph, UObject* InSubject);
+	void HandleModifiedEvent(ERigVMGraphNotifType InNotifType, URigVMGraph* InGraph, UObject* InSubject);
+
+	FString GetValidName(const FString& InPrefix);
+	bool IsValidGraph();
+	bool IsValidNodeForGraph(URigVMNode* InNode);
+	bool IsValidPinForGraph(URigVMPin* InPin);
+	bool IsValidLinkForGraph(URigVMLink* InLink);
+	void AddPinsForStruct(UStruct* InStruct, URigVMNode* InNode, URigVMPin* InParentPin, ERigVMPinDirection InPinDirection, const FString& InDefaultValue);
+	void AddPinsForArray(UArrayProperty* InArrayProperty, URigVMNode* InNode, URigVMPin* InParentPin, ERigVMPinDirection InPinDirection, const TArray<FString>& InDefaultValues);
+	void ConfigurePinFromProperty(UProperty* InProperty, URigVMPin* InOutPin, ERigVMPinDirection InPinDirection = ERigVMPinDirection::Invalid);
+	void ConfigurePinFromPin(URigVMPin* InOutPin, URigVMPin* InPin);
+	virtual bool ShouldStructBeUnfolded(const UStruct* InStruct);
+	virtual bool ShouldPinBeUnfolded(URigVMPin* InPin);
+	void SetPinDefaultValue(URigVMPin* InPin, const FString& InDefaultValue, bool bResizeArrays, bool bUndo, bool bMergeUndoAction);
+	static TArray<FString> SplitDefaultValue(const FString& InDefaultValue);
+	URigVMPin* InsertArrayPin(URigVMPin* ArrayPin, int32 InIndex, const FString& InDefaultValue, bool bUndo);
+	bool RemovePin(URigVMPin* InPinToRemove, bool bUndo);
+	UProperty* FindPropertyForPin(const FString& InPinPath);
+	bool AddLink(URigVMPin* OutputPin, URigVMPin* InputPin, bool bUndo);
+	bool BreakLink(URigVMPin* OutputPin, URigVMPin* InputPin, bool bUndo);
+	bool BreakAllLinks(URigVMPin* Pin, bool bAsInput, bool bUndo);
+	void BreakAllLinksRecursive(URigVMPin* Pin, bool bAsInput, bool bTowardsParent, bool bUndo);
+
+	void ReportWarning(const FString& InMessage);
+	void ReportError(const FString& InMessage);
+
+	template <typename FmtType, typename... Types>
+	void ReportWarningf(const FmtType& Fmt, Types... Args)
+	{
+		ReportWarning(FString::Printf(Fmt, Args...));
+	}
+
+	template <typename FmtType, typename... Types>
+	void ReportErrorf(const FmtType& Fmt, Types... Args)
+	{
+		ReportError(FString::Printf(Fmt, Args...));
+	}
+
+	void DetachLinksFromPinObjects();
+	void ReattachLinksToPinObjects();
+
+#if WITH_EDITOR
+	void RepopulatePinsOnNode(URigVMNode* InNode);
+#endif
+
+	static FLinearColor GetColorFromMetadata(const FString& InMetadata);
+
+	UPROPERTY()
+	URigVMGraph* Graph;
+
+	UPROPERTY(transient)
+	URigVMActionStack* ActionStack;
+
+	bool bReportWarningsAndErrors;
+	static int32 sRemovedPinIndex;
+
+	friend class URigVMGraph;
+	friend class URigVMActionStack;
+	friend class URigVMCompiler;
+};
+
diff --git a/Engine/Source/Editor/RigVMEditor/Public/RigVMModel/RigVMControllerActions.h b/Engine/Source/Editor/RigVMEditor/Public/RigVMModel/RigVMControllerActions.h
new file mode 100644
--- /dev/null
+++ b/Engine/Source/Editor/RigVMEditor/Public/RigVMModel/RigVMControllerActions.h
@@ -0,0 +1,812 @@
+// Copyright 1998-2019 Epic Games, Inc. All Rights Reserved.
+
+#pragma once
+
+#include "UObject/PropertyPortFlags.h"
+#include "RigVMController.h"
+#if WITH_EDITOR
+#include "ScopedTransaction.h"
+#endif
+#include "RigVMControllerActions.generated.h"
+
+/**
+ * ================================================================================
+ * The RigVMController doesn't rely on the transaction system for performing,
+ * tracking, undoing and redoing changes. Instead it uses an action stack which
+ * stores small serialized structs for each occured action. The reason for this 
+ * is the subscription model and Python support: We need the Graph to broadcast
+ * events to all subscribers - independently from where the action is coming from.
+ * This includes UI views, scripting or undo / redo.
+ * Each action supports the concept of 'merging', so multiple color change actions
+ * for example can be merged into a single action. This avoids the need for tracking
+ * action scope - and makes the integration with UI code simple.
+ * The Controller's ActionStack integrates into the editor's transaction stack
+ * using the ActionIndex property: Transactions on that property cause the 
+ * actionstack to consecutively undo or redo actions until the expected stack
+ * size is reached.
+ * ================================================================================
+ */
+
+struct FRigVMBaseAction;
+
+/**
+ * The action key is used for serializing and storing an action in the stack,
+ * or within another action.
+ */
+USTRUCT()
+struct FRigVMActionKey
+{
+	GENERATED_BODY()
+
+	UPROPERTY()
+	FString ScriptStructPath;
+
+	UPROPERTY()
+	FString ExportedText;
+
+	template<class ActionType>
+	void Set(const ActionType& InAction)
+	{
+		UScriptStruct* ScriptStruct = ActionType::StaticStruct();
+		FRigVMActionKey Key;
+		ScriptStructPath = ScriptStruct->GetPathName();
+		ScriptStruct->ExportText(ExportedText, &InAction, nullptr, nullptr, PPF_None, nullptr);
+	}
+};
+
+/**
+ * The action wrapper is used to extract an action from a serialized key.
+ */
+struct FRigVMActionWrapper
+{
+public:
+	FRigVMActionWrapper(const FRigVMActionKey& Key);
+	~FRigVMActionWrapper();
+
+	FRigVMBaseAction* GetAction();
+	FString ExportText();
+
+private:
+	FRigVMActionWrapper(const FRigVMActionWrapper& Other) = delete;
+	FRigVMActionWrapper& operator = (const FRigVMActionWrapper& Other) = delete;
+
+	UScriptStruct* ScriptStruct;
+	TArray<uint8> Data;
+};
+
+/**
+ * The base action is the base struct for all actions, and provides
+ * access to sub actions, merge functionality as well as undo and redo
+ * base implementations.
+ */
+USTRUCT()
+struct FRigVMBaseAction
+{
+	GENERATED_BODY()
+
+	// Default constructor
+	FRigVMBaseAction()
+		 : Title(TEXT("ACtion"))
+	{
+	}
+
+	// Default destructor
+	virtual ~FRigVMBaseAction() {};
+
+	// Returns the title of the action - used for the Edit menu's undo / redo
+	virtual FString GetTitle() const { return Title; }
+
+	// Trys to merge the action with another action and 
+	// returns true if successfull.
+	virtual bool Merge(const FRigVMBaseAction* Other);
+
+	// Un-does the action and returns true if successfull.
+	virtual bool Undo(URigVMController* InController);
+
+	// Re-does the action and returns true if successfull.
+	virtual bool Redo(URigVMController* InController);
+
+	// Adds a child / sub action to this one
+	template<class ActionType>
+	void AddAction(const ActionType& InAction)
+	{
+		FRigVMActionKey Key;
+		Key.Set<ActionType>(InAction);
+		SubActions.Add(Key);
+	}
+
+	UPROPERTY()
+	FString Title;
+
+	UPROPERTY()
+	TArray<FRigVMActionKey> SubActions;
+};
+
+/**
+ * The Action Stack can be used to track actions happening on a
+ * Graph. Currently the only owner of the ActionStack is the Controller.
+ * Actions can be added to the stack, or they can be understood as
+ * scopes / brackets. For this you can use BeginAction / EndAction / CancelAction
+ * to open / close a bracket. Open brackets automatically record additional
+ * actions occuring during the bracket's lifetime.
+ */
+UCLASS()
+class URigVMActionStack : public UObject
+{
+	GENERATED_BODY()
+
+public:
+
+	// Begins an action and opens a bracket / scope.
+	template<class T>
+	void BeginAction(T& InAction)
+	{
+		if (CurrentActions.Num() > 0)
+		{
+			// catch erroreous duplicate calls to begin action
+			ensure(CurrentActions.Last() != &InAction);
+		}
+		CurrentActions.Add((FRigVMBaseAction*)&InAction);
+	}
+
+	// Ends an action and closes a bracket / scope.
+	template<class ActionType>
+	void EndAction(ActionType& InAction, bool bPerformMerge = false)
+	{
+		ensure(CurrentActions.Num() > 0);
+		ensure((FRigVMBaseAction*)&InAction == CurrentActions.Last());
+		CurrentActions.Pop();
+		AddAction(InAction, bPerformMerge);
+	}
+
+	// Cancels an action, closes a bracket / scope and discards all 
+	// actions to this point.
+	template<class ActionType>
+	void CancelAction(ActionType& InAction)
+	{
+		ensure(CurrentActions.Num() > 0);
+		ensure((FRigVMBaseAction*)&InAction == CurrentActions.Last());
+		CurrentActions.Pop();
+	}
+
+	// Adds an action to the stack. Optionally this can perform
+	// a potential merge of this action with the previous action to
+	// compact the stack.
+	template<class ActionType>
+	void AddAction(const ActionType& InAction, bool bPerformMerge = false)
+	{
+		TArray<FRigVMActionKey>* ActionList = &UndoActions;
+		if (CurrentActions.Num() > 0)
+		{
+			ActionList = &CurrentActions[CurrentActions.Num()-1]->SubActions;
+		}
+
+		bool bMergeIfPossible = bPerformMerge;
+		if (bMergeIfPossible)
+		{
+			bMergeIfPossible = false;
+			if (ActionList->Num() > 0 && InAction.SubActions.Num() == 0)
+			{
+				if (ActionList->Last().ScriptStructPath == ActionType::StaticStruct()->GetPathName())
+				{
+					FRigVMActionWrapper Wrapper((*ActionList)[ActionList->Num() - 1]);
+					if (Wrapper.GetAction()->SubActions.Num() == 0)
+					{
+						bMergeIfPossible = Wrapper.GetAction()->Merge(&InAction);
+						if (bMergeIfPossible)
+						{
+							(*ActionList)[ActionList->Num()-1].ExportedText = Wrapper.ExportText();
+						}
+					}
+				}
+			}
+		}
+
+		if (!bMergeIfPossible)
+		{
+			FRigVMActionKey Key;
+			Key.Set<ActionType>(InAction);
+
+			ActionList->Add(Key);
+
+			if (CurrentActions.Num() == 0)
+			{
+				RedoActions.Reset();
+			}
+
+#if WITH_EDITOR
+			const FScopedTransaction Transaction(FText::FromString(InAction.GetTitle()));
+			SetFlags(RF_Transactional);
+			Modify();
+#endif
+
+			ActionIndex = ActionIndex + 1;
+		}
+	}
+
+	// Opens an undo bracket / scope to record actions into.
+	// This is primary useful for Python.
+	UFUNCTION()
+	bool OpenUndoBracket(const FString& InTitle);
+
+	// Closes an undo bracket / scope.
+	// This is primary useful for Python.
+	UFUNCTION()
+	bool CloseUndoBracket();
+
+	// Pops the last action from the undo stack and perform undo on it.
+	// Note: This should really only be used for unit tests,
+	// use the GEditor's main Undo method instead.
+	UFUNCTION()
+	bool Undo(URigVMController* InController);
+
+	// Pops the last action from the redo stack and perform redo on it.
+	// Note: This should really only be used for unit tests,
+	// use the GEditor's main Redo method instead.
+	UFUNCTION()
+	bool Redo(URigVMController* InController);
+
+#if WITH_EDITOR
+	virtual void PostTransacted(const FTransactionObjectEvent& TransactionEvent) override;
+#endif
+
+private:
+
+	UPROPERTY()
+	int32 ActionIndex;
+
+	UPROPERTY()
+	TArray<FRigVMActionKey> UndoActions;
+
+	UPROPERTY()
+	TArray<FRigVMActionKey> RedoActions;
+
+	TArray<FRigVMBaseAction*> CurrentActions;
+	TArray<FRigVMBaseAction*> BracketActions;
+};
+
+/**
+ * An action which inverses the child actions,
+ * it performs undo on redo and vice versa.
+ */
+USTRUCT()
+struct FRigVMInverseAction : public FRigVMBaseAction
+{
+	GENERATED_BODY()
+
+	virtual ~FRigVMInverseAction() {};
+	virtual bool Undo(URigVMController* InController);
+	virtual bool Redo(URigVMController* InController);
+};
+
+/**
+ * An action adding a struct node to the graph.
+ */
+USTRUCT()
+struct FRigVMAddStructNodeAction : public FRigVMBaseAction
+{
+	GENERATED_BODY()
+
+public:
+
+	FRigVMAddStructNodeAction();
+	FRigVMAddStructNodeAction(URigVMStructNode* InNode);
+	virtual ~FRigVMAddStructNodeAction() {};
+	virtual bool Undo(URigVMController* InController) override;
+	virtual bool Redo(URigVMController* InController) override;
+
+	UPROPERTY()
+	FString ScriptStructPath;
+
+	UPROPERTY()
+	FName MethodName;
+	
+	UPROPERTY()
+	FVector2D Position;
+
+	UPROPERTY()
+	FString NodePath;
+};
+
+/**
+ * An action adding a variable node to the graph.
+ */
+USTRUCT()
+struct FRigVMAddVariableNodeAction : public FRigVMBaseAction
+{
+	GENERATED_BODY()
+
+public:
+
+	FRigVMAddVariableNodeAction();
+	FRigVMAddVariableNodeAction(URigVMVariableNode* InNode);
+	virtual ~FRigVMAddVariableNodeAction() {};
+	virtual bool Undo(URigVMController* InController) override;
+	virtual bool Redo(URigVMController* InController) override;
+
+	UPROPERTY()
+	FName VariableName;
+	
+	UPROPERTY()
+	FString CPPType;
+
+	UPROPERTY()
+	FString ScriptStructPath;
+
+	UPROPERTY()
+	bool bIsGetter;
+
+	UPROPERTY()
+	FString DefaultValue;
+
+	UPROPERTY()
+	FVector2D Position;
+
+	UPROPERTY()
+	FString NodePath;
+};
+
+/**
+ * An action adding a parameter node to the graph.
+ */
+USTRUCT()
+struct FRigVMAddParameterNodeAction : public FRigVMBaseAction
+{
+	GENERATED_BODY()
+
+public:
+
+	FRigVMAddParameterNodeAction();
+	FRigVMAddParameterNodeAction(URigVMParameterNode* InNode);
+	virtual ~FRigVMAddParameterNodeAction() {};
+	virtual bool Undo(URigVMController* InController) override;
+	virtual bool Redo(URigVMController* InController) override;
+
+	UPROPERTY()
+	FName ParameterName;
+	
+	UPROPERTY()
+	FString CPPType;
+
+	UPROPERTY()
+	FString ScriptStructPath;
+
+	UPROPERTY()
+	bool bIsInput;
+
+	UPROPERTY()
+	FString DefaultValue;
+
+	UPROPERTY()
+	FVector2D Position;
+
+	UPROPERTY()
+	FString NodePath;
+};
+
+/**
+ * An action adding a comment node to the graph.
+ */
+USTRUCT()
+struct FRigVMAddCommentNodeAction : public FRigVMBaseAction
+{
+	GENERATED_BODY()
+
+public:
+
+	FRigVMAddCommentNodeAction();
+	FRigVMAddCommentNodeAction(URigVMCommentNode* InNode);
+	virtual ~FRigVMAddCommentNodeAction() {};
+	virtual bool Undo(URigVMController* InController) override;
+	virtual bool Redo(URigVMController* InController) override;
+
+	UPROPERTY()
+	FString CommentText;
+
+	UPROPERTY()
+	FVector2D Position;
+
+	UPROPERTY()
+	FVector2D Size;
+
+	UPROPERTY()
+	FLinearColor Color;
+
+	UPROPERTY()
+	FString NodePath;
+};
+
+/**
+ * An action adding a reroute node to the graph.
+ */
+USTRUCT()
+struct FRigVMAddRerouteNodeAction : public FRigVMBaseAction
+{
+	GENERATED_BODY()
+
+public:
+
+	FRigVMAddRerouteNodeAction();
+	FRigVMAddRerouteNodeAction(URigVMRerouteNode* InNode, const FString& InPinPath, bool bInAsInput);
+	FRigVMAddRerouteNodeAction(URigVMRerouteNode* InNode);
+	virtual ~FRigVMAddRerouteNodeAction() {};
+	virtual bool Undo(URigVMController* InController) override;
+	virtual bool Redo(URigVMController* InController) override;
+
+	UPROPERTY()
+	FString PinPath;
+
+	UPROPERTY()
+	bool bShowAsFullNode;
+
+	UPROPERTY()
+	bool bAsInput;
+
+	UPROPERTY()
+	FVector2D Position;
+
+	UPROPERTY()
+	FString NodePath;
+};
+
+/**
+ * An action removing a node from the graph.
+ */
+USTRUCT()
+struct FRigVMRemoveNodeAction : public FRigVMBaseAction
+{
+	GENERATED_BODY()
+
+public:
+
+	FRigVMRemoveNodeAction() {}
+	FRigVMRemoveNodeAction(URigVMNode* InNode);
+	virtual ~FRigVMRemoveNodeAction() {};
+	virtual bool Undo(URigVMController* InController) override;
+	virtual bool Redo(URigVMController* InController) override;
+
+	UPROPERTY()
+	FRigVMActionKey InverseActionKey;
+};
+
+/**
+ * An action selecting or deselecting a node in the graph.
+ */
+USTRUCT()
+struct FRigVMSelectNodeAction : public FRigVMBaseAction
+{
+	GENERATED_BODY()
+
+public:
+
+	FRigVMSelectNodeAction();
+	FRigVMSelectNodeAction(URigVMNode* InNode);
+	virtual ~FRigVMSelectNodeAction() {};
+	virtual bool Undo(URigVMController* InController) override;
+	virtual bool Redo(URigVMController* InController) override;
+
+	UPROPERTY()
+	FString NodePath;
+
+	UPROPERTY()
+	bool bWasSelected;
+};
+
+/**
+ * An action setting a node's position in the graph.
+ */
+USTRUCT()
+struct FRigVMSetNodePositionAction : public FRigVMBaseAction
+{
+	GENERATED_BODY()
+
+public:
+
+	FRigVMSetNodePositionAction()
+	{
+		OldPosition = NewPosition = FVector2D::ZeroVector;
+	}
+	FRigVMSetNodePositionAction(URigVMNode* InNode, const FVector2D& InNewPosition);
+	virtual ~FRigVMSetNodePositionAction() {};
+	virtual bool Merge(const FRigVMBaseAction* Other);
+	virtual bool Undo(URigVMController* InController) override;
+	virtual bool Redo(URigVMController* InController) override;
+
+	UPROPERTY()
+	FString NodePath;
+
+	UPROPERTY()
+	FVector2D OldPosition;
+
+	UPROPERTY()
+	FVector2D NewPosition;
+};
+
+/**
+ * An action setting a node's size in the graph.
+ */
+USTRUCT()
+struct FRigVMSetNodeSizeAction : public FRigVMBaseAction
+{
+	GENERATED_BODY()
+
+public:
+
+	FRigVMSetNodeSizeAction()
+	{
+		OldSize = NewSize = FVector2D::ZeroVector;
+	}
+	FRigVMSetNodeSizeAction(URigVMNode* InNode, const FVector2D& InNewSize);
+	virtual ~FRigVMSetNodeSizeAction() {};
+	virtual bool Merge(const FRigVMBaseAction* Other);
+	virtual bool Undo(URigVMController* InController) override;
+	virtual bool Redo(URigVMController* InController) override;
+
+	UPROPERTY()
+	FString NodePath;
+
+	UPROPERTY()
+	FVector2D OldSize;
+
+	UPROPERTY()
+	FVector2D NewSize;
+};
+
+/**
+ * An action setting a node's color in the graph.
+ */
+USTRUCT()
+struct FRigVMSetNodeColorAction : public FRigVMBaseAction
+{
+	GENERATED_BODY()
+
+public:
+
+	FRigVMSetNodeColorAction()
+	{
+		OldColor = NewColor = FLinearColor::Black;
+	}
+	FRigVMSetNodeColorAction(URigVMNode* InNode, const FLinearColor& InNewColor);
+	virtual ~FRigVMSetNodeColorAction() {};
+	virtual bool Merge(const FRigVMBaseAction* Other);
+	virtual bool Undo(URigVMController* InController) override;
+	virtual bool Redo(URigVMController* InController) override;
+
+	UPROPERTY()
+	FString NodePath;
+
+	UPROPERTY()
+	FLinearColor OldColor;
+
+	UPROPERTY()
+	FLinearColor NewColor;
+};
+
+/**
+ * An action setting a comment node's text in the graph.
+ */
+USTRUCT()
+struct FRigVMSetCommentTextAction : public FRigVMBaseAction
+{
+	GENERATED_BODY()
+
+public:
+
+	FRigVMSetCommentTextAction() {}
+	FRigVMSetCommentTextAction(URigVMCommentNode* InNode, const FString& InNewText);
+	virtual ~FRigVMSetCommentTextAction() {};
+	virtual bool Undo(URigVMController* InController) override;
+	virtual bool Redo(URigVMController* InController) override;
+
+	UPROPERTY()
+	FString NodePath;
+
+	UPROPERTY()
+	FString OldText;
+
+	UPROPERTY()
+	FString NewText;
+};
+
+/**
+ * An action setting a reroute node's compactness in the graph.
+ */
+USTRUCT()
+struct FRigVMSetRerouteCompactnessAction : public FRigVMBaseAction
+{
+	GENERATED_BODY()
+
+public:
+
+	FRigVMSetRerouteCompactnessAction();
+	FRigVMSetRerouteCompactnessAction(URigVMRerouteNode* InNode, bool InShowAsFullNode);
+	virtual ~FRigVMSetRerouteCompactnessAction() {};
+	virtual bool Undo(URigVMController* InController) override;
+	virtual bool Redo(URigVMController* InController) override;
+
+	UPROPERTY()
+	FString NodePath;
+
+	UPROPERTY()
+	bool OldShowAsFullNode;
+
+	UPROPERTY()
+	bool NewShowAsFullNode;
+};
+
+/**
+ * An action renaming a variable in the graph.
+ */
+USTRUCT()
+struct FRigVMRenameVariableAction : public FRigVMBaseAction
+{
+	GENERATED_BODY()
+
+public:
+
+	FRigVMRenameVariableAction() {}
+	FRigVMRenameVariableAction(const FName& InOldVariableName, const FName& InNewVariableName);
+	virtual ~FRigVMRenameVariableAction() {};
+	virtual bool Undo(URigVMController* InController) override;
+	virtual bool Redo(URigVMController* InController) override;
+
+	UPROPERTY()
+	FString OldVariableName;
+
+	UPROPERTY()
+	FString NewVariableName;
+};
+
+/**
+ * An action renaming a parameter in the graph.
+ */
+USTRUCT()
+struct FRigVMRenameParameterAction : public FRigVMBaseAction
+{
+	GENERATED_BODY()
+
+public:
+
+	FRigVMRenameParameterAction() {}
+	FRigVMRenameParameterAction(const FName& InOldParameterName, const FName& InNewParameterName);
+	virtual ~FRigVMRenameParameterAction() {};
+	virtual bool Undo(URigVMController* InController) override;
+	virtual bool Redo(URigVMController* InController) override;
+
+	UPROPERTY()
+	FString OldParameterName;
+
+	UPROPERTY()
+	FString NewParameterName;
+};
+
+/**
+ * An action setting a pin's default value in the graph.
+ */
+USTRUCT()
+struct FRigVMSetPinDefaultValueAction : public FRigVMBaseAction
+{
+	GENERATED_BODY()
+
+public:
+
+	FRigVMSetPinDefaultValueAction() {}
+	FRigVMSetPinDefaultValueAction(URigVMPin* InPin, const FString& InNewDefaultValue);
+	virtual ~FRigVMSetPinDefaultValueAction() {};
+	virtual bool Merge(const FRigVMBaseAction* Other);
+	virtual bool Undo(URigVMController* InController) override;
+	virtual bool Redo(URigVMController* InController) override;
+
+	UPROPERTY()
+	FString PinPath;
+
+	UPROPERTY()
+	FString OldDefaultValue;
+
+	UPROPERTY()
+	FString NewDefaultValue;
+};
+
+/**
+ * An action inserting a new array pin in the graph.
+ */
+USTRUCT()
+struct FRigVMInsertArrayPinAction : public FRigVMBaseAction
+{
+	GENERATED_BODY()
+
+public:
+
+	FRigVMInsertArrayPinAction()
+		: Index(0)
+	{
+	}
+	FRigVMInsertArrayPinAction(URigVMPin* InArrayPin, int32 InIndex, const FString& InNewDefaultValue);
+	virtual ~FRigVMInsertArrayPinAction() {};
+	virtual bool Undo(URigVMController* InController) override;
+	virtual bool Redo(URigVMController* InController) override;
+
+	UPROPERTY()
+	FString ArrayPinPath;
+
+	UPROPERTY()
+	int32 Index;
+
+	UPROPERTY()
+	FString NewDefaultValue;
+};
+
+/**
+ * An action removing an array pin from the graph.
+ */
+USTRUCT()
+struct FRigVMRemoveArrayPinAction : public FRigVMBaseAction
+{
+	GENERATED_BODY()
+
+public:
+
+	FRigVMRemoveArrayPinAction()
+		: Index(0)
+	{
+	}
+	FRigVMRemoveArrayPinAction(URigVMPin* InArrayElementPin);
+	virtual ~FRigVMRemoveArrayPinAction() {};
+	virtual bool Undo(URigVMController* InController) override;
+	virtual bool Redo(URigVMController* InController) override;
+
+	UPROPERTY()
+	FString ArrayPinPath;
+
+	UPROPERTY()
+	int32 Index;
+
+	UPROPERTY()
+	FString DefaultValue;
+};
+
+/**
+ * An action adding a new link to the graph.
+ */
+USTRUCT()
+struct FRigVMAddLinkAction : public FRigVMBaseAction
+{
+	GENERATED_BODY()
+
+public:
+
+	FRigVMAddLinkAction() {}
+	FRigVMAddLinkAction(URigVMPin* InOutputPin, URigVMPin* InInputPin);
+	virtual ~FRigVMAddLinkAction() {};
+	virtual bool Undo(URigVMController* InController) override;
+	virtual bool Redo(URigVMController* InController) override;
+
+	UPROPERTY()
+	FString OutputPinPath;
+
+	UPROPERTY()
+	FString InputPinPath;
+};
+
+/**
+ * An action removing a link from the graph.
+ */
+USTRUCT()
+struct FRigVMBreakLinkAction : public FRigVMBaseAction
+{
+	GENERATED_BODY()
+
+public:
+
+	FRigVMBreakLinkAction() {}
+	FRigVMBreakLinkAction(URigVMPin* InOutputPin, URigVMPin* InInputPin);
+	virtual ~FRigVMBreakLinkAction() {};
+	virtual bool Undo(URigVMController* InController) override;
+	virtual bool Redo(URigVMController* InController) override;
+
+	UPROPERTY()
+	FString OutputPinPath;
+
+	UPROPERTY()
+	FString InputPinPath;
+};
\ No newline at end of file
diff --git a/Engine/Source/Editor/RigVMEditor/Public/RigVMModel/RigVMGraph.h b/Engine/Source/Editor/RigVMEditor/Public/RigVMModel/RigVMGraph.h
new file mode 100644
--- /dev/null
+++ b/Engine/Source/Editor/RigVMEditor/Public/RigVMModel/RigVMGraph.h
@@ -0,0 +1,107 @@
+// Copyright 1998-2019 Epic Games, Inc. All Rights Reserved.
+
+#pragma once
+
+#include "RigVMNode.h"
+#include "RigVMLink.h"
+#include "RigVMNotifications.h"
+#include "RigVMGraphUtils.h"
+#include "Nodes/RigVMVariableNode.h"
+#include "Nodes/RigVMParameterNode.h"
+#include "RigVMGraph.generated.h"
+
+/**
+ * The Graph represents a Function definition
+ * using Nodes as statements.
+ * Graphs can be compiled into a URigVM using the 
+ * FRigVMCompiler. 
+ * Graphs provide access to its Nodes, Pins and
+ * Links.
+ */
+UCLASS(BlueprintType)
+class RIGVMEDITOR_API URigVMGraph : public UObject
+{
+	GENERATED_BODY()
+
+public:
+
+	// Default constructor
+	URigVMGraph();
+
+	// Returns all of the Nodes within this Graph.
+	UFUNCTION(BlueprintCallable, Category = RigVMGraph)
+	const TArray<URigVMNode*>& GetNodes() const;
+
+	// Returns all of the Links within this Graph.
+	UFUNCTION(BlueprintCallable, Category = RigVMGraph)
+	const TArray<URigVMLink*>& GetLinks() const;
+
+	// Returns a list of unique Variable descriptions within this Graph.
+	// Multiple Variable Nodes can share the same description.
+	UFUNCTION(BlueprintCallable, Category = RigVMGraph)
+	TArray<FRigVMGraphVariableDescription> GetVariableDescriptions() const;
+
+	// Returns a list of unique Parameter descriptions within this Graph.
+	// Multiple Parameter Nodes can share the same description.
+	UFUNCTION(BlueprintCallable, Category = RigVMGraph)
+	TArray<FRigVMGraphParameterDescription> GetParameterDescriptions() const;
+
+	// Returns a Node given its name (or nullptr).
+	UFUNCTION(BlueprintCallable, Category = RigVMGraph)
+	URigVMNode* FindNodeByName(const FName& InNodeName) const;
+
+	// Returns a Node given its path (or nullptr).
+	// (for now this is the same as finding a node by its name.)
+	UFUNCTION(BlueprintCallable, Category = RigVMGraph)
+	URigVMNode* FindNode(const FString& InNodePath) const;
+
+	// Returns a Pin given its path, for example "Node.Color.R".
+	UFUNCTION(BlueprintCallable, Category = RigVMGraph)
+	URigVMPin* FindPin(const FString& InPinPath) const;
+
+	// Returns a link given its string representation,
+	// for example "NodeA.Color.R -> NodeB.Translation.X"
+	UFUNCTION(BlueprintCallable, Category = RigVMGraph)
+	URigVMLink* FindLink(const FString& InLinkPinPathRepresentation) const;
+
+	// Returns true if a Node with a given name is selected.
+	UFUNCTION(BlueprintCallable, Category = RigVMGraph)
+	bool IsNodeSelected(const FName& InNodeName) const;
+
+	// Returns the names of all currently selected Nodes.
+	UFUNCTION(BlueprintCallable, Category = RigVMGraph)
+	const TArray<FName>& GetSelectNodes() const;
+
+	// Returns the modified event, which can be used to 
+	// subscribe to changes happening within the Graph.
+	FRigVMGraphModifiedEvent& OnModified();
+
+#if WITH_EDITORONLY_DATA
+	virtual void PostLoad() override;
+#endif
+
+private:
+
+	virtual bool CanLink(URigVMPin* InSourcePin, URigVMPin* InTargetPin, FString* OutFailureReason = nullptr);
+	void PrepareCycleChecking(URigVMPin* InPin, bool bAsInput);
+	void RepopulatePinLinks();
+
+	FRigVMGraphModifiedEvent ModifiedEvent;
+	void Notify(ERigVMGraphNotifType InNotifType, URigVMGraph* InGraph, UObject* InSubject);
+
+	UPROPERTY()
+	TArray<URigVMNode*> Nodes;
+
+	UPROPERTY()
+	TArray<URigVMLink*> Links;
+
+	UPROPERTY()
+	TArray<FName> SelectedNodes;
+
+	FRigVMGraphUtils Utils;
+
+	bool IsNameAvailable(const FString& InName);
+
+	friend class URigVMController;
+};
+
diff --git a/Engine/Source/Editor/RigVMEditor/Public/RigVMModel/RigVMGraphUtils.h b/Engine/Source/Editor/RigVMEditor/Public/RigVMModel/RigVMGraphUtils.h
new file mode 100644
--- /dev/null
+++ b/Engine/Source/Editor/RigVMEditor/Public/RigVMModel/RigVMGraphUtils.h
@@ -0,0 +1,73 @@
+// Copyright 1998-2019 Epic Games, Inc. All Rights Reserved.
+
+#pragma once
+
+#include "CoreMinimal.h"
+#include "RigVMModel/RigVMNotifications.h"
+
+class URigVMGraph;
+class URigVMNode;
+class URigVMPin;
+
+/**
+ * The GraphUtils helper struct can be used to perform
+ * analysis on the graph, sort the Nodes by execution 
+ * order or determine cycles.
+ */
+class RIGVMEDITOR_API FRigVMGraphUtils
+{
+public:
+
+	// Default constructor
+	FRigVMGraphUtils();
+
+	// Constructor for an ad-hoc cycle check.
+	FRigVMGraphUtils(URigVMGraph* InGraph, URigVMPin* InCycleCheckingPin = nullptr, bool InCycleCheckingPinIsInput = true);
+
+	// Default destructor
+	~FRigVMGraphUtils();
+
+	// Returns the Graph used by this utility.
+	URigVMGraph* GetGraph() { return Graph;  }
+
+	// Returns the current Pin used to perform cycle checking on.
+	URigVMPin* GetCycleCheckPin() { return LastCycleCheckingPin; }
+
+	// Returns true if the Pin for cycle checking was used as an input.
+	bool GetCycleCheckPinIsInput() { return LastCycleCheckingPinIsInput; }
+
+	// Resets the utility and fees all memory.
+	void Reset();
+
+	// Returns the distance for a given node to the output node farthest away
+	int32 GetMaxDistanceToLeafOutput(URigVMNode* InNode) const;
+
+	// Returns true if the graph can be sorted and stores the order of execution in the order array.
+	// Return false if there's a cycle in the graph and stores the cycle in the potentialcycle array.
+	bool TopologicalSort(TArray<URigVMNode*>& OutOrder, TArray<URigVMNode*>& OutPotentialCycle);
+
+	// Finds the first cycle. Returns an empty array if there's no cycle
+	TArray<URigVMNode*> FindCycle();
+
+	// Prepares this utility for cycle checking given a Pin and its usage for a link.
+	void PrepareCycleChecking(URigVMPin* InCycleCheckingPin = nullptr, bool InCycleCheckingPinIsInput = true);
+
+	// Returns true if a Node is on the last determined cycle.
+	bool IsNodeOnCycle(URigVMNode* InNode);
+
+private:
+
+	URigVMGraph* Graph;
+	URigVMPin* LastCycleCheckingPin;
+	bool LastCycleCheckingPinIsInput;
+
+	bool IsNodeCyclic(int32 NodeIndex);
+
+	FDelegateHandle ModifiedHandle;
+	TSet<int32> CycleWhiteList;
+	TSet<int32> CycleGrayList;
+	TSet<int32> CycleBlackList;
+	TMap<int32, int32> CycleDepthTraversal;
+	TArray<URigVMNode*> Cycle;
+	TArray<bool> NodeIsOnCycle;
+};
diff --git a/Engine/Source/Editor/RigVMEditor/Public/RigVMModel/RigVMLink.h b/Engine/Source/Editor/RigVMEditor/Public/RigVMModel/RigVMLink.h
new file mode 100644
--- /dev/null
+++ b/Engine/Source/Editor/RigVMEditor/Public/RigVMModel/RigVMLink.h
@@ -0,0 +1,65 @@
+// Copyright 1998-2019 Epic Games, Inc. All Rights Reserved.
+
+#pragma once
+
+#include "CoreMinimal.h"
+#include "UObject/ObjectMacros.h"
+#include "RigVMPin.h"
+#include "RigVMLink.generated.h"
+
+class URigVMGraph;
+
+/**
+ * The Link represents a connection between two Pins
+ * within a Graph. The Link can be accessed on the 
+ * Graph itself - or through the URigVMPin::GetLinks()
+ * method.
+ */
+UCLASS(BlueprintType)
+class RIGVMEDITOR_API URigVMLink : public UObject
+{
+	GENERATED_BODY()
+
+public:
+
+	// Default constructor
+	URigVMLink()
+	{
+		SourcePin = TargetPin = nullptr;
+	}
+
+	// Serialization override
+	virtual void Serialize(FArchive& Ar) override;
+
+	// Returns the current index of this Link within its owning Graph.
+	UFUNCTION(BlueprintCallable, Category = RigVMLink)
+	int32 GetLinkIndex() const;
+
+	// Returns the Link's owning Graph/
+	UFUNCTION(BlueprintCallable, Category = RigVMLink)
+	URigVMGraph* GetGraph() const;
+
+	// Returns the source Pin of this Link (or nullptr)
+	UFUNCTION(BlueprintCallable, Category = RigVMLink)
+	URigVMPin* GetSourcePin();
+
+	// Returns the target Pin of this Link (or nullptr)
+	UFUNCTION(BlueprintCallable, Category = RigVMLink)
+	URigVMPin* GetTargetPin();
+
+	// Returns a string representation of the Link,
+	// for example: "NodeA.Color.R -> NodeB.Translation.X"
+	UFUNCTION(BlueprintCallable, Category = RigVMLink)
+	FString GetPinPathRepresentation();
+
+private:
+
+	FString SourcePinPath;
+	FString TargetPinPath;
+
+	URigVMPin* SourcePin;
+	URigVMPin* TargetPin;
+
+	friend class URigVMController;
+};
+
diff --git a/Engine/Source/Editor/RigVMEditor/Public/RigVMModel/RigVMPin.h b/Engine/Source/Editor/RigVMEditor/Public/RigVMModel/RigVMPin.h
new file mode 100644
--- /dev/null
+++ b/Engine/Source/Editor/RigVMEditor/Public/RigVMModel/RigVMPin.h
@@ -0,0 +1,225 @@
+// Copyright 1998-2019 Epic Games, Inc. All Rights Reserved.
+
+#pragma once
+
+#include "CoreMinimal.h"
+#include "UObject/ObjectMacros.h"
+#include "RigVMPin.generated.h"
+
+class URigVMGraph;
+class URigVMNode;
+class URigVMPin;
+class URigVMLink;
+
+/**
+ * The Pin Direction is used to differentiate different kinds of 
+ * pins in the data flow graph - inputs, outputs etc.
+ */
+UENUM(BlueprintType)
+enum class ERigVMPinDirection : uint8
+{
+	Input, // A const input value
+	Output, // A mutable output value
+	IO, // A mutable input and output value
+	Visible, // A const value that cannot be connected to
+	Hidden, // A mutable hidden value (used for interal state)
+	Invalid // The max value for this enum - used for guarding.
+};
+
+
+/**
+ * The Pin represents a single connector / pin on a node in the RigVM model.
+ * Pins can be connected based on rules. Pins also provide access to a 'PinPath',
+ * which essentially represents . separated list of names to reach the pin within
+ * the owning graph. PinPaths are unique.
+ * In comparison to the EdGraph Pin the URigVMPin supports the concept of 'SubPins',
+ * so child / parent relationships between pins. A FVector Pin for example might
+ * have its X, Y and Z components as SubPins. Array Pins will have its elements as
+ * SubPins, and so on.
+ * A URigVMPin is owned solely by a URigVMNode.
+ */
+UCLASS(BlueprintType)
+class RIGVMEDITOR_API URigVMPin : public UObject
+{
+	GENERATED_BODY()
+
+public:
+
+	// Splits a PinPath at the start, so for example "Node.Color.R" becomes "Node" and "Color.R"
+	static bool SplitPinPathAtStart(const FString& InPinPath, FString& LeftMost, FString& Right);
+
+	// Splits a PinPath at the start, so for example "Node.Color.R" becomes "Node.Color" and "R"
+	static bool SplitPinPathAtEnd(const FString& InPinPath, FString& Left, FString& RightMost);
+
+	// Splits a PinPath into all segments, so for example "Node.Color.R" becomes ["Node", "Color", "R"]
+	static bool SplitPinPath(const FString& InPinPath, TArray<FString>& Parts);
+
+	// Joins a PinPath from to segments, so for example "Node.Color" and "R" becomes "Node.Color.R"
+	static FString JoinSplitPath(const FString& Left, const FString& Right);
+
+	// Joins a PinPath from to segments, so for example ["Node", "Color", "R"] becomes "Node.Color.R"
+	static FString JoinSplitPath(const TArray<FString>& InParts);
+
+	// Default constructor
+	URigVMPin();
+
+	// Returns a . separated path containing all names of the pin and its owners,
+	// this includes the node name, for example "Node.Color.R"
+	UFUNCTION(BlueprintCallable, Category = RigVMPin)
+	FString GetPinPath() const;
+
+	// Returns a . separated path containing all names of the pin within its main
+	// memory owner / storage. This is typically used to create an offset pointer
+	// within memory (FRigVMRegisterOffset).
+	// So for example for a PinPath such as "Node.Transform.Translation.X" the 
+	// corresponding SegmentPath is "Translation.X", since the transform is the
+	// storage / memory.
+	UFUNCTION(BlueprintCallable, Category = RigVMPin)
+	FString GetSegmentPath() const;
+
+	// Returns the direction of the pin
+	UFUNCTION(BlueprintCallable, Category = RigVMPin)
+	ERigVMPinDirection GetDirection() const;
+
+	// Returns true if the pin is a constant value / literal
+	UFUNCTION(BlueprintCallable, Category = RigVMPin)
+	bool IsConstant() const;
+
+	// Returns true if the data type of the Pin is a struct
+	UFUNCTION(BlueprintCallable, Category = RigVMPin)
+	bool IsStruct() const;
+
+	// Returns true if the Pin is a SubPin within a struct
+	UFUNCTION(BlueprintCallable, Category = RigVMPin)
+	bool IsStructMember() const;
+
+	// Returns true if the data type of the Pin is an array
+	UFUNCTION(BlueprintCallable, Category = RigVMPin)
+	bool IsArray() const;
+
+	// Returns true if the Pin is a SubPin within an array
+	UFUNCTION(BlueprintCallable, Category = RigVMPin)
+	bool IsArrayElement() const;
+
+	// Returns the index of the Pin within the node / parent Pin
+	UFUNCTION(BlueprintCallable, Category = RigVMPin)
+	int32 GetPinIndex() const;
+
+	// Returns the number of elements within an array Pin
+	UFUNCTION(BlueprintCallable, Category = RigVMPin)
+	int32 GetArraySize() const;
+
+	// Returns the C++ data type of the pin
+	UFUNCTION(BlueprintCallable, Category = RigVMPin)
+	FString GetCPPType() const;
+
+	// Returns the C++ data type of an element of the Pin array
+	UFUNCTION(BlueprintCallable, Category = RigVMPin)
+	FString GetArrayElementCppType() const;
+
+	// Returns true if the C++ data type is FString or FName
+	UFUNCTION(BlueprintCallable, Category = RigVMPin)
+	bool IsStringType() const;
+
+	// Returns the default value of the Pin as a string.
+	// Note that this value is computed based on the Pin's
+	// SubPins - so for example for a FVector typed Pin
+	// the default value is actually composed out of the
+	// default values of the X, Y and Z SubPins.
+	UFUNCTION(BlueprintCallable, Category = RigVMPin)
+	FString GetDefaultValue() const;
+
+	// Returns the name of a custom widget to be used
+	// for editing the Pin.
+	UFUNCTION(BlueprintCallable, Category = RigVMPin)
+	FName GetCustomWidgetName() const;
+
+	// Returns the struct of the data type of the Pin,
+	// or nullptr otherwise.
+	UFUNCTION(BlueprintCallable, Category = RigVMPin)
+	UScriptStruct* GetScriptStruct() const;
+
+	// Returns the parent Pin - or nullptr if the Pin
+	// is nested directly below a node.
+	UFUNCTION(BlueprintCallable, Category = RigVMPin)
+	URigVMPin* GetParentPin() const;
+
+	// Returns the top-most parent Pin, so for example
+	// for "Node.Transform.Translation.X" this returns
+	// the Pin for "Node.Transform".
+	UFUNCTION(BlueprintCallable, Category = RigVMPin)
+	URigVMPin* GetRootPin() const;
+
+	// Returns all of the SubPins of this one.
+	UFUNCTION(BlueprintCallable, Category = RigVMPin)
+	const TArray<URigVMPin*>& GetSubPins() const;
+
+	// Returns a SubPin given a name / path or nullptr.
+	UFUNCTION(BlueprintCallable, Category = RigVMPin)
+	URigVMPin* FindSubPin(const FString& InPinPath) const;
+
+	// Returns true if this Pin is linked to another Pin
+	UFUNCTION(BlueprintCallable, Category = RigVMPin)
+	bool IsLinkedTo(URigVMPin* InPin) const;
+
+	// Returns all of the links linked to this Pin.
+	UFUNCTION(BlueprintCallable, Category = RigVMPin)
+	const TArray<URigVMLink*>& GetLinks() const;
+
+	// Returns all of the linked source Pins,
+	// using this Pin as the target.
+	UFUNCTION(BlueprintCallable, Category = RigVMPin)
+	TArray<URigVMPin*> GetLinkedSourcePins() const;
+
+	// Returns all of the linked target Pins,
+	// using this Pin as the source.
+	UFUNCTION(BlueprintCallable, Category = RigVMPin)
+	TArray<URigVMPin*> GetLinkedTargetPins() const;
+
+	// Returns the node of this Pin.
+	UFUNCTION(BlueprintCallable, Category = RigVMPin)
+	URigVMNode* GetNode() const;
+
+	// Returns the graph of this Pin.
+	UFUNCTION(BlueprintCallable, Category = RigVMPin)
+	URigVMGraph* GetGraph() const;
+
+	// Returns true is the two provided source and target Pins
+	// can be linked to one another.
+	static bool CanLink(URigVMPin* InSourcePin, URigVMPin* InTargetPin, FString* OutFailureReason = nullptr);
+
+private:
+
+	void SetNameFromIndex();
+
+	UPROPERTY()
+	ERigVMPinDirection Direction;
+
+	UPROPERTY()
+	bool bIsConstant;
+
+	UPROPERTY()
+	FString CPPType;
+
+	UPROPERTY(transient)
+	UScriptStruct* ScriptStruct;
+
+	UPROPERTY()
+	FName ScriptStructPath;
+
+	UPROPERTY()
+	FString DefaultValue;
+
+	UPROPERTY()
+	FName CustomWidgetName;
+
+	UPROPERTY()
+	TArray<URigVMPin*> SubPins;
+
+	UPROPERTY(transient)
+	TArray<URigVMLink*> Links;
+
+	friend class URigVMController;
+	friend class URigVMGraph;
+};
+
diff --git a/Engine/Source/Editor/RigVMEditor/RigVMEditor.Build.cs b/Engine/Source/Editor/RigVMEditor/RigVMEditor.Build.cs
new file mode 100644
--- /dev/null
+++ b/Engine/Source/Editor/RigVMEditor/RigVMEditor.Build.cs
@@ -0,0 +1,28 @@
+// Copyright 1998-2019 Epic Games, Inc. All Rights Reserved.
+
+using UnrealBuildTool;
+
+public class RigVMEditor : ModuleRules
+{
+    public RigVMEditor(ReadOnlyTargetRules Target) : base(Target)
+    {
+        PublicDependencyModuleNames.AddRange(
+            new string[] {
+                "Core",
+                "CoreUObject",
+                "Engine",
+                "RigVM",
+            }
+        );
+
+        if (Target.bBuildEditor == true)
+        {
+            PrivateDependencyModuleNames.AddRange(
+                new string[]
+                {
+                    "UnrealEd",
+                }
+            );
+        }
+    }
+}
diff --git a/Engine/Source/Editor/UnrealEd/Private/EditorViewportClient.cpp b/Engine/Source/Editor/UnrealEd/Private/EditorViewportClient.cpp
--- a/Engine/Source/Editor/UnrealEd/Private/EditorViewportClient.cpp
+++ b/Engine/Source/Editor/UnrealEd/Private/EditorViewportClient.cpp
@@ -5524,6 +5524,19 @@ bool FEditorViewportClient::ProcessScreenShots(FViewport* InViewport)
 				BitmapSize = FIntPoint(NewWidth, NewHeight);
 			}
 
+			if (GIsAutomationTesting)
+			{
+				// Under automation test, the screenshot is highjacked and sent to be compared
+				TArray<FColor> BitmapForCompare(Bitmap);
+				if (!bWriteAlpha)
+				{
+					// Set full alpha on the bitmap
+					for (FColor& Pixel : BitmapForCompare) { Pixel.A = 255; }
+				}
+
+				FScreenshotRequest::OnScreenshotCaptured().Broadcast(BitmapSize.X, BitmapSize.Y, MoveTemp(BitmapForCompare));
+			}
+
 			TUniquePtr<FImageWriteTask> ImageTask = MakeUnique<FImageWriteTask>();
 			ImageTask->PixelData = MakeUnique<TImagePixelData<FColor>>(BitmapSize, MoveTemp(Bitmap));
 
diff --git a/Engine/Source/Programs/UnrealHeaderTool/Resources/UHTDebugging.manifest b/Engine/Source/Programs/UnrealHeaderTool/Resources/UHTDebugging.manifest
--- a/Engine/Source/Programs/UnrealHeaderTool/Resources/UHTDebugging.manifest
+++ b/Engine/Source/Programs/UnrealHeaderTool/Resources/UHTDebugging.manifest
@@ -34,7 +34,9 @@
 				".\\UHTDebugging\\TestInterfaceObject.h",
 				".\\UHTDebugging\\StructOnlyHeader.h",
 				".\\UHTDebugging\\EnumOnlyHeader.h",
-				".\\UHTDebugging\\TestObjectSerializers.h"
+				".\\UHTDebugging\\TestObjectSerializers.h",
+				".\\UHTDebugging\\RigVMStructHeader.h",
+				".\\UHTDebugging\\DocumentationObjects.h"
 			],
 			"PublicHeaders": [
 			],
diff --git a/Engine/Source/Programs/UnrealHeaderTool/Resources/UHTDebugging/DocumentationObjects.h b/Engine/Source/Programs/UnrealHeaderTool/Resources/UHTDebugging/DocumentationObjects.h
new file mode 100644
--- /dev/null
+++ b/Engine/Source/Programs/UnrealHeaderTool/Resources/UHTDebugging/DocumentationObjects.h
@@ -0,0 +1,16 @@
+// Copyright 1998-2019 Epic Games, Inc. All Rights Reserved.
+
+#pragma once
+
+#include "DocumentationObjects.generated.h"
+
+/**
+ * A class to test the documentation policy 
+ * against. For now we'll focus on the main class
+ * tooltip only (this comment).
+ */
+UCLASS(meta=(DocumentationPolicy = "Strict"))
+class UClassToDocument : public UObject
+{
+	GENERATED_BODY()
+};
diff --git a/Engine/Source/Runtime/AnimGraphRuntime/Private/AnimNodes/AnimNode_BlendListBase.cpp b/Engine/Source/Runtime/AnimGraphRuntime/Private/AnimNodes/AnimNode_BlendListBase.cpp
--- a/Engine/Source/Runtime/AnimGraphRuntime/Private/AnimNodes/AnimNode_BlendListBase.cpp
+++ b/Engine/Source/Runtime/AnimGraphRuntime/Private/AnimNodes/AnimNode_BlendListBase.cpp
@@ -3,6 +3,8 @@
 #include "AnimNodes/AnimNode_BlendListBase.h"
 #include "AnimationRuntime.h"
 #include "Animation/BlendProfile.h"
+#include "Animation/AnimInstanceProxy.h"
+#include "Animation/AnimNode_Inertialization.h"
 
 /////////////////////////////////////////////////////
 // FAnimNode_BlendListBase
@@ -99,7 +101,29 @@ void FAnimNode_BlendListBase::Update_AnyThread(const FAnimationUpdateContext& Co
 			// scale by the weight difference since we want always consistency:
 			// - if you're moving from 0 to full weight 1, it will use the normal blend time
 			// - if you're moving from 0.5 to full weight 1, it will get there in half the time
-			const float RemainingBlendTime = LastChildIndexIsInvalid ? 0.0f : ( BlendTime[ChildIndex] * WeightDifference );
+			float RemainingBlendTime;
+			if (LastChildIndexIsInvalid)
+			{
+				RemainingBlendTime = 0.0f;
+			}
+			else if (TransitionType == EBlendListTransitionType::Inertialization)
+			{
+				FAnimNode_Inertialization* InertializationNode = Context.GetAncestor<FAnimNode_Inertialization>();
+				if (InertializationNode)
+				{
+					InertializationNode->Request(BlendTime[ChildIndex]);
+				}
+				else
+				{
+					UE_LOG(LogAnimation, Error, TEXT("FAnimNode_BlendListBase: No inertialization context found for transition"));
+				}
+				
+				RemainingBlendTime = 0.0f;
+			}
+			else
+			{
+				RemainingBlendTime = BlendTime[ChildIndex] * WeightDifference;
+			}
 
 			for (int32 i = 0; i < RemainingBlendTimes.Num(); ++i)
 			{
diff --git a/Engine/Source/Runtime/AnimGraphRuntime/Private/AnimNodes/AnimNode_CopyPoseFromMesh.cpp b/Engine/Source/Runtime/AnimGraphRuntime/Private/AnimNodes/AnimNode_CopyPoseFromMesh.cpp
--- a/Engine/Source/Runtime/AnimGraphRuntime/Private/AnimNodes/AnimNode_CopyPoseFromMesh.cpp
+++ b/Engine/Source/Runtime/AnimGraphRuntime/Private/AnimNodes/AnimNode_CopyPoseFromMesh.cpp
@@ -19,8 +19,7 @@ void FAnimNode_CopyPoseFromMesh::Initialize_AnyThread(const FAnimationInitialize
 	FAnimNode_Base::Initialize_AnyThread(Context);
 
 	// Initial update of the node, so we dont have a frame-delay on setup
-	FAnimationUpdateContext UpdateContext(Context.AnimInstanceProxy, 0.0f);
-	GetEvaluateGraphExposedInputs().Execute(UpdateContext);
+	GetEvaluateGraphExposedInputs().Execute(Context);
 
 	RefreshMeshComponent(Context.AnimInstanceProxy->GetSkelMeshComponent());
 }
diff --git a/Engine/Source/Runtime/AnimGraphRuntime/Private/AnimNodes/AnimNode_Slot.cpp b/Engine/Source/Runtime/AnimGraphRuntime/Private/AnimNodes/AnimNode_Slot.cpp
--- a/Engine/Source/Runtime/AnimGraphRuntime/Private/AnimNodes/AnimNode_Slot.cpp
+++ b/Engine/Source/Runtime/AnimGraphRuntime/Private/AnimNodes/AnimNode_Slot.cpp
@@ -40,7 +40,7 @@ void FAnimNode_Slot::Update_AnyThread(const FAnimationUpdateContext& Context)
 	const bool bUpdateSource = bAlwaysUpdateSourcePose || FAnimWeight::IsRelevant(WeightData.SourceWeight);
 	if (bUpdateSource)
 	{
-		const float SourceWeight = bAlwaysUpdateSourcePose ? FAnimWeight::GetSmallestRelevantWeight() : WeightData.SourceWeight;
+		const float SourceWeight = FMath::Max(FAnimWeight::GetSmallestRelevantWeight(), WeightData.SourceWeight);
 		Source.Update(Context.FractionalWeight(SourceWeight));
 	}
 }
diff --git a/Engine/Source/Runtime/AnimGraphRuntime/Private/BoneControllers/AnimNode_RigidBody.cpp b/Engine/Source/Runtime/AnimGraphRuntime/Private/BoneControllers/AnimNode_RigidBody.cpp
--- a/Engine/Source/Runtime/AnimGraphRuntime/Private/BoneControllers/AnimNode_RigidBody.cpp
+++ b/Engine/Source/Runtime/AnimGraphRuntime/Private/BoneControllers/AnimNode_RigidBody.cpp
@@ -132,8 +132,7 @@ void FAnimNode_RigidBody::UpdateComponentPose_AnyThread(const FAnimationUpdateCo
 		{
 			// Create a new context with zero deltatime to freeze time in rest of the graph.
 			// This will be used to capture a frozen pose.
-			FAnimationUpdateContext FrozenContext = Context;
-			FrozenContext.FractionalWeightAndTime(1.f, 0.f);
+			FAnimationUpdateContext FrozenContext = Context.FractionalWeightAndTime(1.f, 0.f);
 
 			Super::UpdateComponentPose_AnyThread(FrozenContext);
 		}
diff --git a/Engine/Source/Runtime/AnimGraphRuntime/Private/BoneControllers/AnimNode_RigidBody_Chaos.cpp b/Engine/Source/Runtime/AnimGraphRuntime/Private/BoneControllers/AnimNode_RigidBody_Chaos.cpp
--- a/Engine/Source/Runtime/AnimGraphRuntime/Private/BoneControllers/AnimNode_RigidBody_Chaos.cpp
+++ b/Engine/Source/Runtime/AnimGraphRuntime/Private/BoneControllers/AnimNode_RigidBody_Chaos.cpp
@@ -132,8 +132,7 @@ void FAnimNode_RigidBody_Chaos::UpdateComponentPose_AnyThread(const FAnimationUp
 		{
 			// Create a new context with zero deltatime to freeze time in rest of the graph.
 			// This will be used to capture a frozen pose.
-			FAnimationUpdateContext FrozenContext = Context;
-			FrozenContext.FractionalWeightAndTime(1.f, 0.f);
+			FAnimationUpdateContext FrozenContext = Context.FractionalWeightAndTime(1.f, 0.f);
 
 			Super::UpdateComponentPose_AnyThread(FrozenContext);
 		}
diff --git a/Engine/Source/Runtime/AnimGraphRuntime/Public/AnimNodes/AnimNode_BlendListBase.h b/Engine/Source/Runtime/AnimGraphRuntime/Public/AnimNodes/AnimNode_BlendListBase.h
--- a/Engine/Source/Runtime/AnimGraphRuntime/Public/AnimNodes/AnimNode_BlendListBase.h
+++ b/Engine/Source/Runtime/AnimGraphRuntime/Public/AnimNodes/AnimNode_BlendListBase.h
@@ -12,6 +12,13 @@
 class UBlendProfile;
 class UCurveFloat;
 
+UENUM()
+enum class EBlendListTransitionType : uint8
+{
+	StandardBlend,
+	Inertialization
+};
+
 // Blend list node; has many children
 USTRUCT(BlueprintInternalUseOnly)
 struct ANIMGRAPHRUNTIME_API FAnimNode_BlendListBase : public FAnimNode_Base
@@ -24,6 +31,9 @@ public:
 	UPROPERTY(EditAnywhere, EditFixedSize, BlueprintReadWrite, Category=Config, meta=(PinShownByDefault))
 	TArray<float> BlendTime;
 
+	UPROPERTY(EditAnywhere, BlueprintReadWrite, Category=Config)
+	EBlendListTransitionType TransitionType;
+
 	UPROPERTY(EditAnywhere, Category=BlendType)
 	EAlphaBlendOption BlendType;
 
@@ -55,7 +65,8 @@ protected:
 
 public:	
 	FAnimNode_BlendListBase()
-		: BlendType(EAlphaBlendOption::Linear)
+		: TransitionType(EBlendListTransitionType::StandardBlend)
+		, BlendType(EAlphaBlendOption::Linear)
 		, bResetChildOnActivation(false)
 		, LastActiveChildIndex(0)
 		, CustomBlendCurve(nullptr)
diff --git a/Engine/Source/Runtime/Core/Private/Math/UnrealMath.cpp b/Engine/Source/Runtime/Core/Private/Math/UnrealMath.cpp
--- a/Engine/Source/Runtime/Core/Private/Math/UnrealMath.cpp
+++ b/Engine/Source/Runtime/Core/Private/Math/UnrealMath.cpp
@@ -646,6 +646,12 @@ void FQuat::ToSwingTwist(const FVector& InTwistAxis, FQuat& OutSwing, FQuat& Out
 	OutSwing = *this * OutTwist.Inverse();
 }
 
+float FQuat::GetTwistAngle(const FVector& TwistAxis) const
+{
+	float XYZ = FVector::DotProduct(TwistAxis, FVector(X, Y, Z));
+	return FMath::UnwindRadians(2.0f * FMath::Atan2(XYZ, W));
+}
+
 FMatrix FRotationAboutPointMatrix::Make(const FQuat& Rot, const FVector& Origin)
 {
 	return FRotationAboutPointMatrix(Rot.Rotator(), Origin);
diff --git a/Engine/Source/Runtime/Core/Public/Math/Quat.h b/Engine/Source/Runtime/Core/Public/Math/Quat.h
--- a/Engine/Source/Runtime/Core/Public/Math/Quat.h
+++ b/Engine/Source/Runtime/Core/Public/Math/Quat.h
@@ -328,10 +328,19 @@ public:
 	 * @param InTwistAxis Axis to use for decomposition
 	 * @param OutSwing swing component quaternion
 	 * @param OutTwist Twist component quaternion
-	 * @warning assumes normalised quaternion and twist axis
+	 * @warning assumes normalized quaternion and twist axis
 	 */
 	CORE_API void ToSwingTwist(const FVector& InTwistAxis, FQuat& OutSwing, FQuat& OutTwist) const;
 
+	/**
+	 * Get the twist angle (in radians) for a specified axis
+	 *
+	 * @param TwistAxis Axis to use for decomposition
+	 * @return Twist angle (in radians)
+	 * @warning assumes normalized quaternion and twist axis
+	 */
+	CORE_API float GetTwistAngle(const FVector& TwistAxis) const;
+
 	/**
 	 * Rotate a vector by this quaternion.
 	 *
diff --git a/Engine/Source/Runtime/Core/Public/Misc/AutomationTest.h b/Engine/Source/Runtime/Core/Public/Misc/AutomationTest.h
--- a/Engine/Source/Runtime/Core/Public/Misc/AutomationTest.h
+++ b/Engine/Source/Runtime/Core/Public/Misc/AutomationTest.h
@@ -1053,6 +1053,9 @@ public:
 	/** Gets the C++ name of the test. */
 	FString GetTestName() const { return TestName; }
 
+	/** Gets the parameter context of the test. */
+	FString GetTestContext() const { return TestParameterContext; }
+
 	/**
 	 * Pure virtual method; returns the number of participants for this test
 	 *
@@ -1563,6 +1566,9 @@ protected:
 	 */
 	virtual FString GetBeautifiedTestName() const = 0;
 
+	/** Sets the parameter context of the test. */
+	virtual void SetTestContext(FString Context) { TestParameterContext = Context; }
+
 protected:
 
 	//Flag to indicate if this is a complex task
@@ -1577,6 +1583,9 @@ protected:
 	/** Name of the test */
 	FString TestName;
 
+	/** Context of the test */
+	FString TestParameterContext;
+
 	/** Info related to the last execution of this test */
 	FAutomationTestExecutionInfo ExecutionInfo;
 
